SwitchboardListener: emit NvAPI error messages to log using string representations rather than error codes

This uses NvAPI_GetErrorMessage() to turn any error codes returned by NvAPI functions
into human-readable names. This results in log messages that look like this:

LogWindows: Error: NvAPI_Initialize failed. Error: NVAPI_LIBRARY_NOT_FOUND

Rather than like this:

LogWindows: Error: NvAPI_Initialize failed. Error code: -2

diff --git a/Engine/Source/Programs/SwitchboardListener/Private/SwitchboardListener.cpp b/Engine/Source/Programs/SwitchboardListener/Private/SwitchboardListener.cpp
--- a/Engine/Source/Programs/SwitchboardListener/Private/SwitchboardListener.cpp
+++ b/Engine/Source/Programs/SwitchboardListener/Private/SwitchboardListener.cpp
@@ -207,7 +207,9 @@ FSwitchboardListener::FSwitchboardListener(const FSwitchboardCommandLineOptions&
 		const NvAPI_Status Result = NvAPI_Initialize();
 		if (Result != NVAPI_OK)
 		{
-			UE_LOG(LogSwitchboard, Fatal, TEXT("NvAPI_Initialize failed. Error code: %d"), Result);
+			NvAPI_ShortString ErrorString;
+			NvAPI_GetErrorMessage(Result, ErrorString);
+			UE_LOG(LogSwitchboard, Fatal, TEXT("NvAPI_Initialize failed. Error: %s"), ANSI_TO_TCHAR(ErrorString));
 		}
 	}
 
@@ -1166,7 +1168,9 @@ static void FillOutSyncTopologies(TArray<FSyncTopo>& SyncTopos)
 
 			if (Result != NVAPI_OK)
 			{
-				UE_LOG(LogSwitchboard, Warning, TEXT("NvAPI_GSync_GetTopology failed. Error code: %d"), Result);
+				NvAPI_ShortString ErrorString;
+				NvAPI_GetErrorMessage(Result, ErrorString);
+				UE_LOG(LogSwitchboard, Warning, TEXT("NvAPI_GSync_GetTopology failed. Error: %s"), ANSI_TO_TCHAR(ErrorString));
 				continue;
 			}
 		}
@@ -1196,7 +1200,9 @@ static void FillOutSyncTopologies(TArray<FSyncTopo>& SyncTopos)
 
 			if (Result != NVAPI_OK)
 			{
-				UE_LOG(LogSwitchboard, Warning, TEXT("NvAPI_GSync_GetTopology failed. Error code: %d"), Result);
+				NvAPI_ShortString ErrorString;
+				NvAPI_GetErrorMessage(Result, ErrorString);
+				UE_LOG(LogSwitchboard, Warning, TEXT("NvAPI_GSync_GetTopology failed. Error: %s"), ANSI_TO_TCHAR(ErrorString));
 				continue;
 			}
 		}
@@ -1265,7 +1271,9 @@ static void FillOutSyncTopologies(TArray<FSyncTopo>& SyncTopos)
 
 			if (Result != NVAPI_OK)
 			{
-				UE_LOG(LogSwitchboard, Warning, TEXT("NvAPI_GSync_GetStatusParameters failed. Error code: %d"), Result);
+				NvAPI_ShortString ErrorString;
+				NvAPI_GetErrorMessage(Result, ErrorString);
+				UE_LOG(LogSwitchboard, Warning, TEXT("NvAPI_GSync_GetStatusParameters failed. Error: %s"), ANSI_TO_TCHAR(ErrorString));
 				continue;
 			}
 
@@ -1284,7 +1292,9 @@ static void FillOutSyncTopologies(TArray<FSyncTopo>& SyncTopos)
 
 			if (Result != NVAPI_OK)
 			{
-				UE_LOG(LogSwitchboard, Warning, TEXT("NvAPI_GSync_GetControlParameters failed. Error code: %d"), Result);
+				NvAPI_ShortString ErrorString;
+				NvAPI_GetErrorMessage(Result, ErrorString);
+				UE_LOG(LogSwitchboard, Warning, TEXT("NvAPI_GSync_GetControlParameters failed. Error: %s"), ANSI_TO_TCHAR(ErrorString));
 				continue;
 			}
 
@@ -1321,7 +1331,9 @@ static void FillOutDriverVersion(FSyncStatus& SyncStatus)
 
 	if (Result != NVAPI_OK)
 	{
-		UE_LOG(LogSwitchboard, Warning, TEXT("NvAPI_SYS_GetDriverAndBranchVersion failed. Error code: %d"), Result);
+		NvAPI_ShortString ErrorString;
+		NvAPI_GetErrorMessage(Result, ErrorString);
+		UE_LOG(LogSwitchboard, Warning, TEXT("NvAPI_SYS_GetDriverAndBranchVersion failed. Error: %s"), ANSI_TO_TCHAR(ErrorString));
 		return;
 	}
 
@@ -1366,7 +1378,9 @@ static void FillOutMosaicTopologies(TArray<FMosaicTopo>& MosaicTopos)
 
 		if (Result != NVAPI_OK)
 		{
-			UE_LOG(LogSwitchboard, Warning, TEXT("NvAPI_Mosaic_EnumDisplayGrids failed. Error code: %d"), Result);
+			NvAPI_ShortString ErrorString;
+			NvAPI_GetErrorMessage(Result, ErrorString);
+			UE_LOG(LogSwitchboard, Warning, TEXT("NvAPI_Mosaic_EnumDisplayGrids failed. Error: %s"), ANSI_TO_TCHAR(ErrorString));
 			return;
 		}
 	}
@@ -1384,7 +1398,9 @@ static void FillOutMosaicTopologies(TArray<FMosaicTopo>& MosaicTopos)
 
 		if (Result != NVAPI_OK)
 		{
-			UE_LOG(LogSwitchboard, Warning, TEXT("NvAPI_Mosaic_EnumDisplayGrids failed. Error code: %d"), Result);
+			NvAPI_ShortString ErrorString;
+			NvAPI_GetErrorMessage(Result, ErrorString);
+			UE_LOG(LogSwitchboard, Warning, TEXT("NvAPI_Mosaic_EnumDisplayGrids failed. Error: %s"), ANSI_TO_TCHAR(ErrorString));
 			return;
 		}
 
@@ -1610,7 +1626,9 @@ static void FillOutPhysicalGpuStats(FSyncStatus& SyncStatus)
 	NvAPI_Status NvResult = NvAPI_EnumPhysicalGPUs(PhysicalGpuHandles.GetData(), &PhysicalGpuCount);
 	if (NvResult != NVAPI_OK)
 	{
-		UE_LOG(LogSwitchboard, Warning, TEXT("NvAPI_EnumPhysicalGPUs failed. Error code: %d"), NvResult);
+		NvAPI_ShortString ErrorString;
+		NvAPI_GetErrorMessage(NvResult, ErrorString);
+		UE_LOG(LogSwitchboard, Warning, TEXT("NvAPI_EnumPhysicalGPUs failed. Error: %s"), ANSI_TO_TCHAR(ErrorString));
 		return;
 	}
 
@@ -1626,13 +1644,21 @@ static void FillOutPhysicalGpuStats(FSyncStatus& SyncStatus)
 
 		if (LhsBusResult != NVAPI_OK || RhsBusResult != NVAPI_OK)
 		{
-			UE_LOG(LogSwitchboard, Warning, TEXT("NvAPI_GPU_GetBusId failed. Error codes: %d, %d"), LhsBusResult, RhsBusResult);
+			NvAPI_ShortString LhsErrorString;
+			NvAPI_ShortString RhsErrorString;
+			NvAPI_GetErrorMessage(LhsBusResult, LhsErrorString);
+			NvAPI_GetErrorMessage(RhsBusResult, RhsErrorString);
+			UE_LOG(LogSwitchboard, Warning, TEXT("NvAPI_GPU_GetBusId failed. Errors: %s, %s"), ANSI_TO_TCHAR(LhsErrorString), ANSI_TO_TCHAR(RhsErrorString));
 			return false;
 		}
 
 		if (LhsSlotResult != NVAPI_OK || RhsSlotResult != NVAPI_OK)
 		{
-			UE_LOG(LogSwitchboard, Warning, TEXT("NvAPI_GPU_GetBusSlotId failed. Error codes: %d, %d"), LhsSlotResult, RhsSlotResult);
+			NvAPI_ShortString LhsErrorString;
+			NvAPI_ShortString RhsErrorString;
+			NvAPI_GetErrorMessage(LhsSlotResult, LhsErrorString);
+			NvAPI_GetErrorMessage(RhsSlotResult, RhsErrorString);
+			UE_LOG(LogSwitchboard, Warning, TEXT("NvAPI_GPU_GetBusSlotId failed. Errors: %s, %s"), ANSI_TO_TCHAR(LhsErrorString), ANSI_TO_TCHAR(RhsErrorString));
 			return false;
 		}
 
