SwitchboardListener: add handling to support systems where NvAPI is unavailable

Instead of exiting with a fatal error if the NvAPI cannot be initialized,
SwitchboardListener now stores a bool to track the success or failure
and bypasses any functionality that involves invoking the NvAPI if
NvAPI initialization failed. This should allow SwitchboardListener to
execute on systems without an NVIDIA GPU, albeit with somewhat
reduced functionality.

diff --git a/Engine/Source/Programs/SwitchboardListener/Private/SwitchboardListener.cpp b/Engine/Source/Programs/SwitchboardListener/Private/SwitchboardListener.cpp
--- a/Engine/Source/Programs/SwitchboardListener/Private/SwitchboardListener.cpp
+++ b/Engine/Source/Programs/SwitchboardListener/Private/SwitchboardListener.cpp
@@ -198,6 +198,7 @@ FSwitchboardListener::FSwitchboardListener(const FSwitchboardCommandLineOptions&
 	: Options(InOptions)
 	, SocketListener(nullptr)
 	, CpuMonitor(MakeShared<FCpuUtilizationMonitor, ESPMode::ThreadSafe>())
+	, bIsNvAPIInitialized(false)
 	, CachedMosaicToposLock(MakeShared<FRWLock, ESPMode::ThreadSafe>())
 	, CachedMosaicTopos(MakeShared<TArray<FMosaicTopo>, ESPMode::ThreadSafe>())
 {
@@ -205,15 +206,19 @@ FSwitchboardListener::FSwitchboardListener(const FSwitchboardCommandLineOptions&
 	// initialize NvAPI
 	{
 		const NvAPI_Status Result = NvAPI_Initialize();
-		if (Result != NVAPI_OK)
+		if (Result == NVAPI_OK)
+		{
+			bIsNvAPIInitialized = true;
+
+			FillOutMosaicTopologies(*CachedMosaicTopos);
+		}
+		else
 		{
 			NvAPI_ShortString ErrorString;
 			NvAPI_GetErrorMessage(Result, ErrorString);
-			UE_LOG(LogSwitchboard, Fatal, TEXT("NvAPI_Initialize failed. Error: %s"), ANSI_TO_TCHAR(ErrorString));
+			UE_LOG(LogSwitchboard, Error, TEXT("NvAPI_Initialize failed. Error: %s"), ANSI_TO_TCHAR(ErrorString));
 		}
 	}
-
-	FillOutMosaicTopologies(*CachedMosaicTopos);
 #endif // PLATFORM_WINDOWS
 
 	const FIPv4Address DefaultIp = FIPv4Address(0, 0, 0, 0);
@@ -1793,13 +1798,18 @@ bool FSwitchboardListener::Task_GetSyncStatus(const FSwitchboardGetSyncStatusTas
 	MessageFuture.Future = Async(EAsyncExecution::ThreadPool,
 		[
 			SyncStatus,
+			IsNvAPIInitialized=bIsNvAPIInitialized,
 			CpuMonitor=CpuMonitor,
 			CachedMosaicToposLock=CachedMosaicToposLock,
 			CachedMosaicTopos=CachedMosaicTopos
 		]() {
-			FillOutDriverVersion(SyncStatus.Get());
 			FillOutTaskbarAutoHide(SyncStatus.Get());
-			FillOutSyncTopologies(SyncStatus->SyncTopos);
+
+			if (IsNvAPIInitialized)
+			{
+				FillOutDriverVersion(SyncStatus.Get());
+				FillOutSyncTopologies(SyncStatus->SyncTopos);
+			}
 
 			{
 				FReadScopeLock Lock(*CachedMosaicToposLock);
@@ -1816,7 +1826,10 @@ bool FSwitchboardListener::Task_GetSyncStatus(const FSwitchboardGetSyncStatusTas
 			const FPlatformMemoryStats MemStats = FPlatformMemory::GetStats();
 			SyncStatus->AvailablePhysicalMemory = MemStats.AvailablePhysical;
 
-			FillOutPhysicalGpuStats(SyncStatus.Get());
+			if (IsNvAPIInitialized)
+			{
+				FillOutPhysicalGpuStats(SyncStatus.Get());
+			}
 
 			return CreateSyncStatusMessage(SyncStatus.Get());
 		}
@@ -1835,6 +1848,14 @@ bool FSwitchboardListener::Task_GetSyncStatus(const FSwitchboardGetSyncStatusTas
 bool FSwitchboardListener::Task_RefreshMosaics(const FSwitchboardRefreshMosaicsTask& InRefreshMosaicsTask)
 {
 #if PLATFORM_WINDOWS
+	if (!bIsNvAPIInitialized)
+	{
+		SendMessage(
+			CreateTaskDeclinedMessage(InRefreshMosaicsTask, "NvAPI not supported", {}),
+			InRefreshMosaicsTask.Recipient);
+		return false;
+	}
+
 	// Reject request if an equivalent one is already in our future
 	if (EquivalentTaskFutureExists(InRefreshMosaicsTask.GetEquivalenceHash()))
 	{
diff --git a/Engine/Source/Programs/SwitchboardListener/Private/SwitchboardListener.h b/Engine/Source/Programs/SwitchboardListener/Private/SwitchboardListener.h
--- a/Engine/Source/Programs/SwitchboardListener/Private/SwitchboardListener.h
+++ b/Engine/Source/Programs/SwitchboardListener/Private/SwitchboardListener.h
@@ -154,6 +154,8 @@ private:
 	TArray<FSwitchboardMessageFuture> MessagesFutures;
 	TSharedPtr<FCpuUtilizationMonitor, ESPMode::ThreadSafe> CpuMonitor;
 
+	bool bIsNvAPIInitialized;
+
 	TSharedPtr<FRWLock, ESPMode::ThreadSafe> CachedMosaicToposLock;
 	TSharedPtr<TArray<FMosaicTopo>, ESPMode::ThreadSafe> CachedMosaicTopos;
 
