SwitchboardListener: add a command-line switch to bypass minimizing the console window on launch

diff --git a/Engine/Source/Programs/SwitchboardListener/Private/SwitchboardListener.cpp b/Engine/Source/Programs/SwitchboardListener/Private/SwitchboardListener.cpp
--- a/Engine/Source/Programs/SwitchboardListener/Private/SwitchboardListener.cpp
+++ b/Engine/Source/Programs/SwitchboardListener/Private/SwitchboardListener.cpp
@@ -147,6 +147,11 @@ FSwitchboardCommandLineOptions FSwitchboardCommandLineOptions::FromString(const
 		OutOptions.OutputVersion = true;
 	}
 
+	if (Switches.Contains(TEXT("noMinimizeOnLaunch")))
+	{
+		OutOptions.MinimizeOnLaunch = false;
+	}
+
 	if (SwitchPairs.Contains(TEXT("ip")))
 	{
 		FIPv4Address ParseAddr;
@@ -173,6 +178,11 @@ FString FSwitchboardCommandLineOptions::ToString(bool bIncludeRedeploy /* = fals
 {
 	TArray<FString> Args;
 
+	if (!MinimizeOnLaunch)
+	{
+		Args.Add(TEXT("-noMinimizeOnLaunch"));
+	}
+
 	if (Address.IsSet())
 	{
 		Args.Add(FString::Printf(TEXT("-ip=%s"), *Address.GetValue().ToString()));
diff --git a/Engine/Source/Programs/SwitchboardListener/Private/SwitchboardListener.h b/Engine/Source/Programs/SwitchboardListener/Private/SwitchboardListener.h
--- a/Engine/Source/Programs/SwitchboardListener/Private/SwitchboardListener.h
+++ b/Engine/Source/Programs/SwitchboardListener/Private/SwitchboardListener.h
@@ -32,6 +32,7 @@ class FTcpListener;
 struct FSwitchboardCommandLineOptions
 {
 	bool OutputVersion = false;
+	bool MinimizeOnLaunch = true;
 
 	TOptional<FIPv4Address> Address;
 	TOptional<uint16> Port;
diff --git a/Engine/Source/Programs/SwitchboardListener/Private/SwitchboardListenerMain.cpp b/Engine/Source/Programs/SwitchboardListener/Private/SwitchboardListenerMain.cpp
--- a/Engine/Source/Programs/SwitchboardListener/Private/SwitchboardListenerMain.cpp
+++ b/Engine/Source/Programs/SwitchboardListener/Private/SwitchboardListenerMain.cpp
@@ -281,7 +281,10 @@ INT32_MAIN_INT32_ARGC_TCHAR_ARGV()
 	}
 
 #if PLATFORM_WINDOWS
-	ShowWindow(GetConsoleWindow(), SW_MINIMIZE);
+	if (Options.MinimizeOnLaunch)
+	{
+		ShowWindow(GetConsoleWindow(), SW_MINIMIZE);
+	}
 #endif
 
 	const bool bListenerResult = RunSwitchboardListener(Options);
