Switchboard: Implement override for Listener inactive timeout, and coarse hitch detection and logging for both Switchboard and the Listener. Bump Listener version to 1.7.0.

Also fixed:
 - Logic error in the ListenerClient keepalive that would lead to keepalives being sent at >= 10 Hz rather than the intended 1 Hz.
 - UTF8_TO_TCHAR conversion and logging of buffers that may not be null terminated.
 - Uninitialized "Recipient" in flip monitor process, leading to "Trying to send message to disconnected client 205.205.205.205:52685" (endpoint = 0xCDCDCDCD...).

diff --git a/Engine/Plugins/VirtualProduction/Switchboard/Source/Switchboard/switchboard/devices/unreal/plugin_unreal.py b/Engine/Plugins/VirtualProduction/Switchboard/Source/Switchboard/switchboard/devices/unreal/plugin_unreal.py
--- a/Engine/Plugins/VirtualProduction/Switchboard/Source/Switchboard/switchboard/devices/unreal/plugin_unreal.py
+++ b/Engine/Plugins/VirtualProduction/Switchboard/Source/Switchboard/switchboard/devices/unreal/plugin_unreal.py
@@ -401,6 +401,15 @@ class DeviceUnreal(Device):
             value=switchboard_application.RsyncServer.DEFAULT_PORT,
             tool_tip='Port number on which the rsync server should listen.'
         ),
+        'listener_inactive_timeout': IntSetting(
+            attr_name='listener_inactive_timeout',
+            nice_name='Listener Timeout',
+            value=5,
+            tool_tip=(
+                'Tells the connected Listener to wait at least N seconds '
+                'between network messages before considering the connection '
+                'to Switchboard lost and closing it with a timeout error.')
+        ),
     }
 
     unreal_started_signal = QtCore.Signal()
@@ -1844,6 +1853,12 @@ class DeviceUnreal(Device):
                     f"{self.name} already running {prog.name} {prog.puuid}")
                 self.do_program_running_update(prog=prog)
 
+        # override listener "inactive" timeout
+        _, msg = message_protocol.create_set_inactive_timeout_message(
+            DeviceUnreal.csettings['listener_inactive_timeout'].get_value()
+        )
+        self.unreal_client.send_message(msg)
+
         # request roles and changelists
         self._request_roles_file()
         self._request_project_changelist_number()
diff --git a/Engine/Plugins/VirtualProduction/Switchboard/Source/Switchboard/switchboard/devices/unreal/version_helpers.py b/Engine/Plugins/VirtualProduction/Switchboard/Source/Switchboard/switchboard/devices/unreal/version_helpers.py
--- a/Engine/Plugins/VirtualProduction/Switchboard/Source/Switchboard/switchboard/devices/unreal/version_helpers.py
+++ b/Engine/Plugins/VirtualProduction/Switchboard/Source/Switchboard/switchboard/devices/unreal/version_helpers.py
@@ -3,7 +3,7 @@
 import typing
 
 
-LISTENER_COMPATIBLE_VERSION = (1,6)
+LISTENER_COMPATIBLE_VERSION = (1,7)
 
 # Version where redeploy support was initially introduced.
 LISTENER_MIN_REDEPLOY_VERSION = (1,4,0)
diff --git a/Engine/Plugins/VirtualProduction/Switchboard/Source/Switchboard/switchboard/listener_client.py b/Engine/Plugins/VirtualProduction/Switchboard/Source/Switchboard/switchboard/listener_client.py
--- a/Engine/Plugins/VirtualProduction/Switchboard/Source/Switchboard/switchboard/listener_client.py
+++ b/Engine/Plugins/VirtualProduction/Switchboard/Source/Switchboard/switchboard/listener_client.py
@@ -31,6 +31,9 @@ class ListenerClient(object):
     The other, legacy (VCS/file) delegates are each passed different (or no)
     arguments; for details, see `route_message()`.
     '''
+
+    KEEPALIVE_INTERVAL_SEC = 1.0
+
     def __init__(self, ip_address, port, buffer_size=1024):
         self.ip_address = ip_address
         self.port = port
@@ -143,9 +146,17 @@ class ListenerClient(object):
             self.close_socket = True
             self.handle_connection_thread.join()
 
+    def update_last_activity(self):
+        now = datetime.datetime.now()
+        delta_sec = (now - self.last_activity).total_seconds()
+        self.last_activity = now
+        if delta_sec > (self.KEEPALIVE_INTERVAL_SEC * 1.5):
+            LOGGER.warning(
+                f'Hitch detected; {delta_sec:.1f} seconds since last keepalive'
+                f' to {self.ip_address}')
+
     def handle_connection(self):
         buffer = []
-        keepalive_timeout = 1.0
 
         while self.is_connected:
             try:
@@ -158,17 +169,17 @@ class ListenerClient(object):
 
                 while len(self.message_queue):
                     message_bytes = self.message_queue.pop()
+                    self.update_last_activity()
                     self.socket.sendall(message_bytes)
-                    self.last_activity = datetime.datetime.now()
                 
                 for rs in read_sockets:
                     received_data = rs.recv(self.buffer_size).decode()
                     self.process_received_data(buffer, received_data)
 
                 delta = datetime.datetime.now() - self.last_activity
-
-                if delta.total_seconds() > keepalive_timeout:
+                if delta.total_seconds() > self.KEEPALIVE_INTERVAL_SEC:
                     _, msg = message_protocol.create_keep_alive_message()
+                    self.update_last_activity()
                     self.socket.sendall(msg)
 
                 if self.close_socket and len(self.message_queue) == 0:
diff --git a/Engine/Plugins/VirtualProduction/Switchboard/Source/Switchboard/switchboard/message_protocol.py b/Engine/Plugins/VirtualProduction/Switchboard/Source/Switchboard/switchboard/message_protocol.py
--- a/Engine/Plugins/VirtualProduction/Switchboard/Source/Switchboard/switchboard/message_protocol.py
+++ b/Engine/Plugins/VirtualProduction/Switchboard/Source/Switchboard/switchboard/message_protocol.py
@@ -122,3 +122,10 @@ def create_minimize_windows_message():
     message = json.dumps(cmd).encode() + b'\x00'
     return (cmd_id, message)
 
+def create_set_inactive_timeout_message(timeout_seconds: int):
+    cmd_id = uuid.uuid4()
+    cmd = {'command': 'set inactive timeout', 'id': str(cmd_id),
+           'seconds': timeout_seconds}
+    message = json.dumps(cmd).encode() + b'\x00'
+    return (cmd_id, message)
+
diff --git a/Engine/Source/Programs/SwitchboardListener/Private/SwitchboardListener.cpp b/Engine/Source/Programs/SwitchboardListener/Private/SwitchboardListener.cpp
--- a/Engine/Source/Programs/SwitchboardListener/Private/SwitchboardListener.cpp
+++ b/Engine/Source/Programs/SwitchboardListener/Private/SwitchboardListener.cpp
@@ -45,9 +45,12 @@ static void FillOutMosaicTopologies(TArray<FMosaicTopo>& MosaicTopos);
 #endif // PLATFORM_WINDOWS
 
 
+const FIPv4Endpoint FSwitchboardListener::InvalidEndpoint(FIPv4Address::LanBroadcast, 0);
+
+
 namespace
 {
-	const double SecondsUntilInactiveClientDisconnect = 5.0;
+	const double DefaultInactiveTimeoutSeconds = 5.0;
 
 	bool TryFindIdInBrokenMessage(const FString& InMessage, FGuid& OutGuid)
 	{
@@ -291,6 +294,7 @@ bool FSwitchboardListener::Tick()
 			Connections.Add(Connection);
 
 			FIPv4Endpoint ClientEndpoint = Connection.Key;
+			InactiveTimeouts.FindOrAdd(ClientEndpoint, DefaultInactiveTimeoutSeconds);
 			LastActivityTime.FindOrAdd(ClientEndpoint, FPlatformTime::Seconds());
 
 			// Send current state upon connection
@@ -481,6 +485,11 @@ bool FSwitchboardListener::RunScheduledTask(const FSwitchboardTask& InTask)
 			const FSwitchboardMinimizeWindowsTask& Task = static_cast<const FSwitchboardMinimizeWindowsTask&>(InTask);
 			return Task_MinimizeWindows(Task);
 		}
+		case ESwitchboardTaskType::SetInactiveTimeout:
+		{
+			const FSwitchboardSetInactiveTimeoutTask& Task = static_cast<const FSwitchboardSetInactiveTimeoutTask&>(InTask);
+			return Task_SetInactiveTimeout(Task);
+		}
 		default:
 		{
 			static const FString Response = TEXT("Unknown Command detected");
@@ -1551,6 +1560,8 @@ FRunningProcess* FSwitchboardListener::FindOrStartFlipModeMonitorForUUID(const F
 	// The monitor auto-closes when monitored program closes.
 	MonitorProcess->UUID = Process->UUID;
 	MonitorProcess->bUpdateClientsWithStdout = false;
+	MonitorProcess->Recipient = InvalidEndpoint;
+	MonitorProcess->Name = TEXT("flipmode_monitor");
 
 	FlipModeMonitors.Add(MonitorProcess);
 
@@ -1569,7 +1580,8 @@ static void FillOutFlipMode(FSyncStatus& SyncStatus, FRunningProcess* FlipModeMo
 	}
 
 	// Get stdout.
-	const FString StdOut(UTF8_TO_TCHAR(FlipModeMonitor->Output.GetData()));
+	FUTF8ToTCHAR StdoutConv(reinterpret_cast<const ANSICHAR*>(FlipModeMonitor->Output.GetData()), FlipModeMonitor->Output.Num());
+	const FString StdOut(StdoutConv.Get());
 
 	// Clear out the StdOut array.
 	FlipModeMonitor->Output.Empty();
@@ -1958,15 +1970,38 @@ bool FSwitchboardListener::Task_MinimizeWindows(const FSwitchboardMinimizeWindow
 #endif // PLATFORM_WINDOWS
 }
 
+bool FSwitchboardListener::Task_SetInactiveTimeout(const FSwitchboardSetInactiveTimeoutTask& InTimeoutTask)
+{
+	const FIPv4Endpoint& Client = InTimeoutTask.Recipient;
+	const float RequestedTimeout = InTimeoutTask.TimeoutSeconds;
+
+	if (RequestedTimeout < DefaultInactiveTimeoutSeconds)
+	{
+		SendMessage(CreateTaskDeclinedMessage(InTimeoutTask, "Requested timeout too low", {}), Client);
+		return false;
+	}
+
+	float& ClientTimeout = InactiveTimeouts[Client];
+	if (RequestedTimeout != ClientTimeout)
+	{
+		UE_LOG(LogSwitchboard, Display, TEXT("Changing client %s inactive timeout from %.0f to %.0f seconds"),
+			*Client.ToString(), ClientTimeout, RequestedTimeout);
+		ClientTimeout = RequestedTimeout;
+	}
+
+	return true;
+}
+
 void FSwitchboardListener::CleanUpDisconnectedSockets()
 {
 	const double CurrentTime = FPlatformTime::Seconds();
 	for (const TPair<FIPv4Endpoint, double>& LastActivity : LastActivityTime)
 	{
 		const FIPv4Endpoint& Client = LastActivity.Key;
-		if (CurrentTime - LastActivity.Value > SecondsUntilInactiveClientDisconnect)
+		const float ClientTimeout = InactiveTimeouts[Client];
+		if (CurrentTime - LastActivity.Value > ClientTimeout)
 		{
-			UE_LOG(LogSwitchboard, Warning, TEXT("Client %s has been inactive for more than %.1fs -- closing connection"), *Client.ToString(), SecondsUntilInactiveClientDisconnect);
+			UE_LOG(LogSwitchboard, Warning, TEXT("Client %s has been inactive for more than %.1fs -- closing connection"), *Client.ToString(), ClientTimeout);
 			TUniquePtr<FSwitchboardDisconnectTask> DisconnectTask = MakeUnique<FSwitchboardDisconnectTask>(FGuid(), Client);
 			DisconnectTasks.Enqueue(MoveTemp(DisconnectTask));
 		}
@@ -1994,6 +2029,7 @@ void FSwitchboardListener::DisconnectClient(const FIPv4Endpoint& InClientEndpoin
 	const FString Client = InClientEndpoint.ToString();
 	UE_LOG(LogSwitchboard, Display, TEXT("Client %s disconnected"), *Client);
 	Connections.Remove(InClientEndpoint);
+	InactiveTimeouts.Remove(InClientEndpoint);
 	LastActivityTime.Remove(InClientEndpoint);
 	ReceiveBuffer.Remove(InClientEndpoint);
 }
@@ -2052,11 +2088,13 @@ void FSwitchboardListener::HandleRunningProcesses(TArray<TSharedPtr<FRunningProc
 
 				int32 ReturnCode = 0;
 				FPlatformProcess::GetProcReturnCode(Process->Handle, &ReturnCode);
-				UE_LOG(LogSwitchboard, Display, TEXT("Process exited with returncode: %d"), ReturnCode);
+				UE_LOG(LogSwitchboard, Display, TEXT("Process %d (%s) exited with returncode: %d"),
+					Process->PID, *Process->Name, ReturnCode);
 
-				if (ReturnCode != 0)
+				if (ReturnCode != 0 && Process->Output.Num() > 0)
 				{
-					UE_LOG(LogSwitchboard, Display, TEXT("Output:\n%s"), UTF8_TO_TCHAR(Process->Output.GetData()));
+					FUTF8ToTCHAR StdoutConv(reinterpret_cast<const ANSICHAR*>(Process->Output.GetData()), Process->Output.Num());
+					UE_LOG(LogSwitchboard, Display, TEXT("Output:\n%s"), StdoutConv.Get());
 				}
 
 				// Notify remote client, which implies that this is a program managed by it.
@@ -2117,6 +2155,11 @@ bool FSwitchboardListener::OnIncomingConnection(FSocket* InSocket, const FIPv4En
 
 bool FSwitchboardListener::SendMessage(const FString& InMessage, const FIPv4Endpoint& InEndpoint)
 {
+	if (InEndpoint == InvalidEndpoint)
+	{
+		return false;
+	}
+
 	if (Connections.Contains(InEndpoint))
 	{
 		TSharedPtr<FSocket> ClientSocket = Connections[InEndpoint];
diff --git a/Engine/Source/Programs/SwitchboardListener/Private/SwitchboardListener.h b/Engine/Source/Programs/SwitchboardListener/Private/SwitchboardListener.h
--- a/Engine/Source/Programs/SwitchboardListener/Private/SwitchboardListener.h
+++ b/Engine/Source/Programs/SwitchboardListener/Private/SwitchboardListener.h
@@ -24,6 +24,7 @@ struct FSwitchboardRefreshMosaicsTask;
 struct FSwitchboardRedeployListenerTask;
 struct FSwitchboardFixExeFlagsTask;
 struct FSwitchboardMinimizeWindowsTask;
+struct FSwitchboardSetInactiveTimeoutTask;
 
 class FInternetAddr;
 class FSocket;
@@ -95,6 +96,8 @@ struct FRedeployStatus
 
 class FSwitchboardListener
 {
+	static const FIPv4Endpoint InvalidEndpoint;
+
 public:
 	explicit FSwitchboardListener(const FSwitchboardCommandLineOptions& InOptions);
 	~FSwitchboardListener();
@@ -124,6 +127,7 @@ private:
 	bool Task_RefreshMosaics(const FSwitchboardRefreshMosaicsTask& InRefreshMosaicsTask);
 	bool Task_FixExeFlags(const FSwitchboardFixExeFlagsTask& InFixExeFlagsTask);
 	bool Task_MinimizeWindows(const FSwitchboardMinimizeWindowsTask& InRefreshMosaicsTask);
+	bool Task_SetInactiveTimeout(const FSwitchboardSetInactiveTimeoutTask& InTimeoutTask);
 
 	bool KillProcessNow(FRunningProcess* InProcess, float SoftKillTimeout = 0.0f);
 	FRunningProcess* FindOrStartFlipModeMonitorForUUID(const FGuid& UUID);
@@ -147,6 +151,7 @@ private:
 	TUniquePtr<FTcpListener> SocketListener;
 	TQueue<TPair<FIPv4Endpoint, TSharedPtr<FSocket>>, EQueueMode::Spsc> PendingConnections;
 	TMap<FIPv4Endpoint, TSharedPtr<FSocket>> Connections;
+	TMap<FIPv4Endpoint, float> InactiveTimeouts;
 	TMap<FIPv4Endpoint, double> LastActivityTime;
 	TMap<FIPv4Endpoint, TArray<uint8>> ReceiveBuffer;
 
diff --git a/Engine/Source/Programs/SwitchboardListener/Private/SwitchboardListenerMain.cpp b/Engine/Source/Programs/SwitchboardListener/Private/SwitchboardListenerMain.cpp
--- a/Engine/Source/Programs/SwitchboardListener/Private/SwitchboardListenerMain.cpp
+++ b/Engine/Source/Programs/SwitchboardListener/Private/SwitchboardListenerMain.cpp
@@ -190,6 +190,11 @@ bool RunSwitchboardListener(const FSwitchboardCommandLineOptions& Options)
 		const double DeltaTime = CurrentTime - LastTime;
 		LastTime = CurrentTime;
 
+		if (DeltaTime > 0.1)
+		{
+			UE_LOG(LogSwitchboard, Warning, TEXT("Hitch detected; %.3f seconds since prior tick"), DeltaTime);
+		}
+
 		FTaskGraphInterface::Get().ProcessThreadUntilIdle(ENamedThreads::GameThread);
 
 		// Pump & Tick objects
diff --git a/Engine/Source/Programs/SwitchboardListener/Private/SwitchboardListenerVersion.h b/Engine/Source/Programs/SwitchboardListener/Private/SwitchboardListenerVersion.h
--- a/Engine/Source/Programs/SwitchboardListener/Private/SwitchboardListenerVersion.h
+++ b/Engine/Source/Programs/SwitchboardListener/Private/SwitchboardListenerVersion.h
@@ -3,5 +3,5 @@
 #pragma once
 
 #define SBLISTENER_VERSION_MAJOR 1
-#define SBLISTENER_VERSION_MINOR 6
+#define SBLISTENER_VERSION_MINOR 7
 #define SBLISTENER_VERSION_PATCH 0
diff --git a/Engine/Source/Programs/SwitchboardListener/Private/SwitchboardProtocol.cpp b/Engine/Source/Programs/SwitchboardListener/Private/SwitchboardProtocol.cpp
--- a/Engine/Source/Programs/SwitchboardListener/Private/SwitchboardProtocol.cpp
+++ b/Engine/Source/Programs/SwitchboardListener/Private/SwitchboardProtocol.cpp
@@ -283,6 +283,16 @@ bool CreateTaskFromCommand(const FString& InCommand, const FIPv4Endpoint& InEndp
 		OutTask = MakeUnique<FSwitchboardMinimizeWindowsTask>(MessageID, InEndpoint);
 		return true;
 	}
+	else if (CommandName == TEXT("set inactive timeout"))
+	{
+		TSharedPtr<FJsonValue> SecondsField = TryGetCommandRequiredField(JsonData, TEXT("seconds"));
+
+		if (SecondsField.IsValid())
+		{
+			OutTask = MakeUnique<FSwitchboardSetInactiveTimeoutTask>(MessageID, InEndpoint, SecondsField->AsNumber());
+			return true;
+		}
+	}
 
 	return false;
 }
diff --git a/Engine/Source/Programs/SwitchboardListener/Private/SwitchboardTasks.h b/Engine/Source/Programs/SwitchboardListener/Private/SwitchboardTasks.h
--- a/Engine/Source/Programs/SwitchboardListener/Private/SwitchboardTasks.h
+++ b/Engine/Source/Programs/SwitchboardListener/Private/SwitchboardTasks.h
@@ -21,6 +21,7 @@ enum class ESwitchboardTaskType : uint8
 	RedeployListener,
 	RefreshMosaics,
 	MinimizeWindows,
+	SetInactiveTimeout,
 };
 
 
@@ -232,3 +233,20 @@ struct FSwitchboardMinimizeWindowsTask : public FSwitchboardTask
 	{}
 };
 
+struct FSwitchboardSetInactiveTimeoutTask : public FSwitchboardTask
+{
+	FSwitchboardSetInactiveTimeoutTask(const FGuid& InTaskId, const FIPv4Endpoint& InEndpoint, double InTimeoutSeconds)
+		: FSwitchboardTask{ ESwitchboardTaskType::SetInactiveTimeout, TEXT("set inactive timeout"), InTaskId, InEndpoint }
+		, TimeoutSeconds(InTimeoutSeconds)
+	{}
+
+	double TimeoutSeconds;
+
+	//~ Begin FSwitchboardTask interface
+	virtual uint32 GetEquivalenceHash() const override
+	{
+		return HashCombine(FSwitchboardTask::GetEquivalenceHash(), GetTypeHash(TimeoutSeconds));
+	}
+	//~ End FSwitchboardTask interface
+};
+
