Switchboard: Add Insights tracing to SwitchboardListener.

diff --git a/Engine/Source/Programs/SwitchboardListener/Private/SwitchboardListener.cpp b/Engine/Source/Programs/SwitchboardListener/Private/SwitchboardListener.cpp
--- a/Engine/Source/Programs/SwitchboardListener/Private/SwitchboardListener.cpp
+++ b/Engine/Source/Programs/SwitchboardListener/Private/SwitchboardListener.cpp
@@ -286,6 +286,8 @@ bool FSwitchboardListener::StopListening()
 
 bool FSwitchboardListener::Tick()
 {
+	SWITCHBOARD_TRACE_CPUPROFILER_EVENT_SCOPE(FSwitchboardListener::Tick);
+
 	// Dequeue pending connections
 	{
 		TPair<FIPv4Endpoint, TSharedPtr<FSocket>> Connection;
@@ -383,6 +385,8 @@ bool FSwitchboardListener::Tick()
 
 bool FSwitchboardListener::ParseIncomingMessage(const FString& InMessage, const FIPv4Endpoint& InEndpoint)
 {
+	SWITCHBOARD_TRACE_CPUPROFILER_EVENT_SCOPE(FSwitchboardListener::ParseIncomingMessage);
+
 	TUniquePtr<FSwitchboardTask> Task;
 	bool bEcho = true;
 	if (CreateTaskFromCommand(InMessage, InEndpoint, Task, bEcho))
@@ -624,6 +628,8 @@ static bool DisableFullscreenOptimizationForProcess(const FRunningProcess* Proce
 
 bool FSwitchboardListener::Task_StartProcess(const FSwitchboardStartTask& InRunTask)
 {
+	SWITCHBOARD_TRACE_CPUPROFILER_EVENT_SCOPE(FSwitchboardListener::Task_StartProcess);
+
 	auto NewProcess = MakeShared<FRunningProcess, ESPMode::ThreadSafe>();
 
 	NewProcess->Recipient = InRunTask.Recipient;
@@ -712,6 +718,8 @@ bool FSwitchboardListener::Task_StartProcess(const FSwitchboardStartTask& InRunT
 
 bool FSwitchboardListener::Task_KillProcess(const FSwitchboardKillTask& KillTask)
 {
+	SWITCHBOARD_TRACE_CPUPROFILER_EVENT_SCOPE(FSwitchboardListener::Task_KillProcess);
+
 	if (EquivalentTaskFutureExists(KillTask.GetEquivalenceHash()))
 	{
 		SendMessage(
@@ -769,6 +777,8 @@ bool FSwitchboardListener::Task_KillProcess(const FSwitchboardKillTask& KillTask
 	const FGuid UUID = KillTask.ProgramID;
 
 	MessageFuture.Future = Async(EAsyncExecution::ThreadPool, [=]() {
+		SWITCHBOARD_TRACE_CPUPROFILER_EVENT_SCOPE_STR("FSwitchboardListener::Task_KillProcess future closure");
+
 		const float SoftKillTimeout = 2.0f;
 
 		const bool bKilledProcess = KillProcessNow(Process.Get(), SoftKillTimeout);
@@ -823,6 +833,8 @@ bool FSwitchboardListener::Task_KillProcess(const FSwitchboardKillTask& KillTask
 
 bool FSwitchboardListener::Task_FixExeFlags(const FSwitchboardFixExeFlagsTask& Task)
 {
+	SWITCHBOARD_TRACE_CPUPROFILER_EVENT_SCOPE(FSwitchboardListener::Task_FixExeFlags);
+
 	if (EquivalentTaskFutureExists(Task.GetEquivalenceHash()))
 	{
 		SendMessage(	
@@ -876,6 +888,8 @@ bool FSwitchboardListener::Task_FixExeFlags(const FSwitchboardFixExeFlagsTask& T
 
 bool FSwitchboardListener::KillProcessNow(FRunningProcess* InProcess, float SoftKillTimeout)
 {
+	SWITCHBOARD_TRACE_CPUPROFILER_EVENT_SCOPE(FSwitchboardListener::KillProcessNow);
+
 	if (InProcess && InProcess->Handle.IsValid() && FPlatformProcess::IsProcRunning(InProcess->Handle))
 	{
 		UE_LOG(LogSwitchboard, Display, TEXT("Killing app with PID %d"), InProcess->PID);
@@ -915,6 +929,8 @@ bool FSwitchboardListener::KillProcessNow(FRunningProcess* InProcess, float Soft
 
 bool FSwitchboardListener::Task_ReceiveFileFromClient(const FSwitchboardReceiveFileFromClientTask& InReceiveFileFromClientTask)
 {
+	SWITCHBOARD_TRACE_CPUPROFILER_EVENT_SCOPE(FSwitchboardListener::Task_ReceiveFileFromClient);
+
 	FString Destination = InReceiveFileFromClientTask.Destination;
 
 	if (Destination.Contains(TEXT("%TEMP%")))
@@ -965,6 +981,8 @@ bool FSwitchboardListener::Task_ReceiveFileFromClient(const FSwitchboardReceiveF
 
 bool FSwitchboardListener::Task_RedeployListener(const FSwitchboardRedeployListenerTask& InRedeployListenerTask)
 {
+	SWITCHBOARD_TRACE_CPUPROFILER_EVENT_SCOPE(FSwitchboardListener::Task_RedeployListener);
+
 	if (RedeployStatus.State != FRedeployStatus::EState::None)
 	{
 		const FString ErrorMsg(TEXT("Redeploy already in progress"));
@@ -1105,6 +1123,8 @@ bool FSwitchboardListener::Task_RedeployListener(const FSwitchboardRedeployListe
 
 void FSwitchboardListener::RollbackRedeploy()
 {
+	SWITCHBOARD_TRACE_CPUPROFILER_EVENT_SCOPE(FSwitchboardListener::RollbackRedeploy);
+
 	// Unwind redeploy actions, in reverse order of corresponding state transitions.
 	UE_LOG(LogSwitchboard, Warning, TEXT("Rolling back redeploy..."));
 
@@ -1155,6 +1175,8 @@ FString FSwitchboardListener::GetIpcSemaphoreName(uint32 ParentPid)
 
 bool FSwitchboardListener::Task_SendFileToClient(const FSwitchboardSendFileToClientTask& InSendFileToClientTask)
 {
+	SWITCHBOARD_TRACE_CPUPROFILER_EVENT_SCOPE(FSwitchboardListener::Task_SendFileToClient);
+
 	FString SourceFilePath = InSendFileToClientTask.Source;
 	FPlatformMisc::NormalizePath(SourceFilePath);
 	FPaths::MakePlatformFilename(SourceFilePath);
@@ -1187,6 +1209,8 @@ static FCriticalSection SwitchboardListenerMutexNvapi;
 #if PLATFORM_WINDOWS
 static void FillOutSyncTopologies(TArray<FSyncTopo>& SyncTopos)
 {
+	SWITCHBOARD_TRACE_CPUPROFILER_EVENT_SCOPE(FillOutSyncTopologies);
+
 	FScopeLock LockNvapi(&SwitchboardListenerMutexNvapi);
 
 	// Normally there is a single sync card. BUT an RTX Server could have more, and we need to account for that.
@@ -1364,6 +1388,8 @@ static void FillOutSyncTopologies(TArray<FSyncTopo>& SyncTopos)
 #if PLATFORM_WINDOWS
 static void FillOutDriverVersion(FSyncStatus& SyncStatus)
 {
+	SWITCHBOARD_TRACE_CPUPROFILER_EVENT_SCOPE(FillOutDriverVersion);
+
 	NvU32 DriverVersion;
 	NvAPI_ShortString BuildBranchString;
 
@@ -1385,6 +1411,8 @@ static void FillOutDriverVersion(FSyncStatus& SyncStatus)
 #if PLATFORM_WINDOWS
 static void FillOutTaskbarAutoHide(FSyncStatus& SyncStatus)
 {
+	SWITCHBOARD_TRACE_CPUPROFILER_EVENT_SCOPE(FillOutTaskbarAutoHide);
+
 	APPBARDATA AppBarData;
 
 	AppBarData.cbSize = sizeof(APPBARDATA);
@@ -1407,6 +1435,8 @@ static void FillOutTaskbarAutoHide(FSyncStatus& SyncStatus)
 #if PLATFORM_WINDOWS
 static void FillOutMosaicTopologies(TArray<FMosaicTopo>& MosaicTopos)
 {
+	SWITCHBOARD_TRACE_CPUPROFILER_EVENT_SCOPE(FillOutMosaicTopologies);
+
 	FScopeLock LockNvapi(&SwitchboardListenerMutexNvapi);
 
 	NvU32 GridCount = 0;
@@ -1467,6 +1497,8 @@ static void FillOutMosaicTopologies(TArray<FMosaicTopo>& MosaicTopos)
 #if PLATFORM_WINDOWS
 FRunningProcess* FSwitchboardListener::FindOrStartFlipModeMonitorForUUID(const FGuid& UUID)
 {
+	SWITCHBOARD_TRACE_CPUPROFILER_EVENT_SCOPE(FSwitchboardListener::FindOrStartFlipModeMonitorForUUID);
+
 	// See if the associated FlipModeMonitor is running
 	{
 		auto* FlipModeMonitorPtr = FlipModeMonitors.FindByPredicate([&](const TSharedPtr<FRunningProcess, ESPMode::ThreadSafe>& FlipMonitor)
@@ -1572,6 +1604,8 @@ FRunningProcess* FSwitchboardListener::FindOrStartFlipModeMonitorForUUID(const F
 #if PLATFORM_WINDOWS
 static void FillOutFlipMode(FSyncStatus& SyncStatus, FRunningProcess* FlipModeMonitor)
 {
+	SWITCHBOARD_TRACE_CPUPROFILER_EVENT_SCOPE(FillOutFlipMode);
+
 	// See if the flip monitor is still there.
 	if (!FlipModeMonitor || !FlipModeMonitor->Handle.IsValid())
 	{
@@ -1621,6 +1655,8 @@ static void FillOutFlipMode(FSyncStatus& SyncStatus, FRunningProcess* FlipModeMo
 #if PLATFORM_WINDOWS
 static void FillOutDisableFullscreenOptimizationForProcess(FSyncStatus& SyncStatus, const FRunningProcess* Process)
 {
+	SWITCHBOARD_TRACE_CPUPROFILER_EVENT_SCOPE(FillOutDisableFullscreenOptimizationForProcess);
+
 	// Reset output array just in case
 	SyncStatus.ProgramLayers.Reset();
 
@@ -1659,6 +1695,8 @@ static void FillOutDisableFullscreenOptimizationForProcess(FSyncStatus& SyncStat
 #if PLATFORM_WINDOWS
 static void FillOutPhysicalGpuStats(FSyncStatus& SyncStatus)
 {
+	SWITCHBOARD_TRACE_CPUPROFILER_EVENT_SCOPE(FillOutPhysicalGpuStats);
+
 	// TODO: Can we somehow use the GPU engine "engtype_3D" perf counters for this instead?
 
 	FScopeLock LockNvapi(&SwitchboardListenerMutexNvapi);
@@ -1792,6 +1830,8 @@ bool FSwitchboardListener::EquivalentTaskFutureExists(uint32 TaskEquivalenceHash
 
 bool FSwitchboardListener::Task_GetSyncStatus(const FSwitchboardGetSyncStatusTask& InGetSyncStatusTask)
 {
+	SWITCHBOARD_TRACE_CPUPROFILER_EVENT_SCOPE(FSwitchboardListener::Task_GetSyncStatus);
+
 #if PLATFORM_WINDOWS
 	// Reject request if an equivalent one is already in our future
 	if (EquivalentTaskFutureExists(InGetSyncStatusTask.GetEquivalenceHash()))
@@ -1841,6 +1881,8 @@ bool FSwitchboardListener::Task_GetSyncStatus(const FSwitchboardGetSyncStatusTas
 			CachedMosaicToposLock=CachedMosaicToposLock,
 			CachedMosaicTopos=CachedMosaicTopos
 		]() {
+			SWITCHBOARD_TRACE_CPUPROFILER_EVENT_SCOPE_STR("FSwitchboardListener::Task_GetSyncStatus future closure");
+
 			FillOutTaskbarAutoHide(SyncStatus.Get());
 
 			if (IsNvAPIInitialized)
@@ -1885,6 +1927,8 @@ bool FSwitchboardListener::Task_GetSyncStatus(const FSwitchboardGetSyncStatusTas
 
 bool FSwitchboardListener::Task_RefreshMosaics(const FSwitchboardRefreshMosaicsTask& InRefreshMosaicsTask)
 {
+	SWITCHBOARD_TRACE_CPUPROFILER_EVENT_SCOPE(FSwitchboardListener::Task_RefreshMosaics);
+
 #if PLATFORM_WINDOWS
 	if (!bIsNvAPIInitialized)
 	{
@@ -1914,6 +1958,8 @@ bool FSwitchboardListener::Task_RefreshMosaics(const FSwitchboardRefreshMosaicsT
 			CachedMosaicTopos=CachedMosaicTopos
 		]()
 		{
+			SWITCHBOARD_TRACE_CPUPROFILER_EVENT_SCOPE_STR("FSwitchboardListener::Task_RefreshMosaics future closure");
+
 			FWriteScopeLock Lock(*CachedMosaicToposLock);
 			CachedMosaicTopos->Reset();
 			FillOutMosaicTopologies(*CachedMosaicTopos);
@@ -1933,8 +1979,9 @@ bool FSwitchboardListener::Task_RefreshMosaics(const FSwitchboardRefreshMosaicsT
 
 bool FSwitchboardListener::Task_MinimizeWindows(const FSwitchboardMinimizeWindowsTask& InTask)
 {
-#if PLATFORM_WINDOWS
+	SWITCHBOARD_TRACE_CPUPROFILER_EVENT_SCOPE(FSwitchboardListener::Task_MinimizeWindows);
 
+#if PLATFORM_WINDOWS
 	// Reject request if an equivalent one is already in our future
 	if (EquivalentTaskFutureExists(InTask.GetEquivalenceHash()))
 	{
@@ -1952,6 +1999,8 @@ bool FSwitchboardListener::Task_MinimizeWindows(const FSwitchboardMinimizeWindow
 	TaskFuture.Future = Async(EAsyncExecution::ThreadPool,
 		[]()
 		{
+			SWITCHBOARD_TRACE_CPUPROFILER_EVENT_SCOPE_STR("FSwitchboardListener::Task_MinimizeWindows future closure");
+
 			if (const HWND HWnd = ::FindWindowA("Shell_TrayWnd", NULL))
 			{
 				::SendMessageA(HWnd, WM_COMMAND, (WPARAM)419, 0);
@@ -1994,6 +2043,8 @@ bool FSwitchboardListener::Task_SetInactiveTimeout(const FSwitchboardSetInactive
 
 void FSwitchboardListener::CleanUpDisconnectedSockets()
 {
+	SWITCHBOARD_TRACE_CPUPROFILER_EVENT_SCOPE(FSwitchboardListener::CleanUpDisconnectedSockets);
+
 	const double CurrentTime = FPlatformTime::Seconds();
 	for (const TPair<FIPv4Endpoint, double>& LastActivity : LastActivityTime)
 	{
@@ -2036,6 +2087,8 @@ void FSwitchboardListener::DisconnectClient(const FIPv4Endpoint& InClientEndpoin
 
 void FSwitchboardListener::HandleStdout(const TSharedPtr<FRunningProcess, ESPMode::ThreadSafe>& Process)
 {
+	SWITCHBOARD_TRACE_CPUPROFILER_EVENT_SCOPE(FSwitchboardListener::HandleStdout);
+
 	TArray<uint8> Output;
 	if (FPlatformProcess::ReadPipeToArray(Process->ReadPipe, Output))
 	{
@@ -2065,6 +2118,8 @@ void FSwitchboardListener::HandleStdout(const TSharedPtr<FRunningProcess, ESPMod
 
 void FSwitchboardListener::HandleRunningProcesses(TArray<TSharedPtr<FRunningProcess, ESPMode::ThreadSafe>>& Processes, bool bNotifyThatProgramEnded)
 {
+	SWITCHBOARD_TRACE_CPUPROFILER_EVENT_SCOPE(FSwitchboardListener::HandleRunningProcesses);
+
 	// Reads pipe and cleans up dead processes from the array.
 	for (auto Iter = Processes.CreateIterator(); Iter; ++Iter)
 	{
@@ -2155,6 +2210,8 @@ bool FSwitchboardListener::OnIncomingConnection(FSocket* InSocket, const FIPv4En
 
 bool FSwitchboardListener::SendMessage(const FString& InMessage, const FIPv4Endpoint& InEndpoint)
 {
+	SWITCHBOARD_TRACE_CPUPROFILER_EVENT_SCOPE(FSwitchboardListener::SendMessage);
+
 	if (InEndpoint == InvalidEndpoint)
 	{
 		return false;
@@ -2180,6 +2237,8 @@ bool FSwitchboardListener::SendMessage(const FString& InMessage, const FIPv4Endp
 
 void FSwitchboardListener::SendMessageFutures()
 {
+	SWITCHBOARD_TRACE_CPUPROFILER_EVENT_SCOPE(FSwitchboardListener::SendMessageFutures);
+
 	for (auto Iter = MessagesFutures.CreateIterator(); Iter; ++Iter)
 	{
 		FSwitchboardMessageFuture& MessageFuture = *Iter;
diff --git a/Engine/Source/Programs/SwitchboardListener/Private/SwitchboardListenerApp.h b/Engine/Source/Programs/SwitchboardListener/Private/SwitchboardListenerApp.h
--- a/Engine/Source/Programs/SwitchboardListener/Private/SwitchboardListenerApp.h
+++ b/Engine/Source/Programs/SwitchboardListener/Private/SwitchboardListenerApp.h
@@ -4,5 +4,35 @@
 
 #include "CoreMinimal.h"
 #include "Logging/LogMacros.h"
+#include "ProfilingDebugging/CpuProfilerTrace.h"
+
 
 DECLARE_LOG_CATEGORY_EXTERN(LogSwitchboard, Verbose, All);
+
+
+#if !defined(SWITCHBOARD_CPUPROFILERTRACE_ENABLED)
+#if CPUPROFILERTRACE_ENABLED
+#define SWITCHBOARD_CPUPROFILERTRACE_ENABLED 1
+#else
+#define SWITCHBOARD_CPUPROFILERTRACE_ENABLED 0
+#endif
+#endif
+
+#if SWITCHBOARD_CPUPROFILERTRACE_ENABLED
+
+#define SWITCHBOARD_TRACE_CPUPROFILER_EVENT_SCOPE(Name) \
+	TRACE_CPUPROFILER_EVENT_SCOPE(Name)
+
+#define SWITCHBOARD_TRACE_CPUPROFILER_EVENT_SCOPE_STR(Name) \
+	TRACE_CPUPROFILER_EVENT_SCOPE_STR(Name)
+
+#define SWITCHBOARD_TRACE_CPUPROFILER_EVENT_SCOPE_TEXT(Name) \
+	TRACE_CPUPROFILER_EVENT_SCOPE_TEXT(Name)
+
+#else
+
+#define SWITCHBOARD_TRACE_CPUPROFILER_EVENT_SCOPE(Name)
+#define SWITCHBOARD_TRACE_CPUPROFILER_EVENT_SCOPE_STR(Name)
+#define SWITCHBOARD_TRACE_CPUPROFILER_EVENT_SCOPE_TEXT(Name)
+
+#endif
diff --git a/Engine/Source/Programs/SwitchboardListener/SwitchboardListener.Build.cs b/Engine/Source/Programs/SwitchboardListener/SwitchboardListener.Build.cs
--- a/Engine/Source/Programs/SwitchboardListener/SwitchboardListener.Build.cs
+++ b/Engine/Source/Programs/SwitchboardListener/SwitchboardListener.Build.cs
@@ -15,6 +15,7 @@ public class SwitchboardListener : ModuleRules
 				"Networking",
 				"Projects", // for LaunchEngineLoop.cpp dependency
 				"JsonUtilities",
+				"TraceLog",
 			}
 		);
 
diff --git a/Engine/Source/Programs/SwitchboardListener/SwitchboardListener.Target.cs b/Engine/Source/Programs/SwitchboardListener/SwitchboardListener.Target.cs
--- a/Engine/Source/Programs/SwitchboardListener/SwitchboardListener.Target.cs
+++ b/Engine/Source/Programs/SwitchboardListener/SwitchboardListener.Target.cs
@@ -21,5 +21,7 @@ public class SwitchboardListenerTarget : TargetRules
 
 		// The listener is meant to be a console application (no window), but on MacOS, to get a proper log console, a full application must be built.
 		bIsBuildingConsoleApplication = Target.Platform != UnrealTargetPlatform.Mac;
+
+		GlobalDefinitions.Add("UE_TRACE_ENABLED=1");
 	}
 }
