Switchboard: Enable crash reporting in SwitchboardListener, and add some newly-required shutdown code to facilitate exiting cleanly. Bump Listener version to 2.0.0 to reflect divergence from any future UE4 long-term support changes.

diff --git a/Engine/Plugins/VirtualProduction/Switchboard/Source/Switchboard/switchboard/devices/unreal/version_helpers.py b/Engine/Plugins/VirtualProduction/Switchboard/Source/Switchboard/switchboard/devices/unreal/version_helpers.py
--- a/Engine/Plugins/VirtualProduction/Switchboard/Source/Switchboard/switchboard/devices/unreal/version_helpers.py
+++ b/Engine/Plugins/VirtualProduction/Switchboard/Source/Switchboard/switchboard/devices/unreal/version_helpers.py
@@ -3,7 +3,7 @@
 import typing
 
 
-LISTENER_COMPATIBLE_VERSION = (1,7)
+LISTENER_COMPATIBLE_VERSION = (2,0)
 
 # Version where redeploy support was initially introduced.
 LISTENER_MIN_REDEPLOY_VERSION = (1,4,0)
@@ -21,7 +21,9 @@ def listener_supports_redeploy(listener_ver: ListenerVersion):
 def listener_is_compatible(listener_ver: ListenerVersion):
     return listener_ver[:2] == LISTENER_COMPATIBLE_VERSION
 
-def listener_ver_from_state_message(state_msg: dict) -> typing.Optional[ListenerVersion]:
+def listener_ver_from_state_message(
+    state_msg: dict
+) -> typing.Optional[ListenerVersion]:
     try:
         version = int(state_msg['version'])
     except (KeyError, ValueError):
diff --git a/Engine/Source/Programs/SwitchboardListener/Private/SwitchboardListenerMain.cpp b/Engine/Source/Programs/SwitchboardListener/Private/SwitchboardListenerMain.cpp
--- a/Engine/Source/Programs/SwitchboardListener/Private/SwitchboardListenerMain.cpp
+++ b/Engine/Source/Programs/SwitchboardListener/Private/SwitchboardListenerMain.cpp
@@ -4,17 +4,19 @@
 #include "SwitchboardListenerApp.h"
 #include "SwitchboardListenerVersion.h"
 
+#include "HAL/ExceptionHandling.h"
 #include "Interfaces/IPv4/IPv4Endpoint.h"
 #include "Misc/ScopeExit.h"
+#include "ProfilingDebugging/TraceAuxiliary.h"
 #include "RequiredProgramMainCPPInclude.h"
 
 IMPLEMENT_APPLICATION(SwitchboardListener, "SwitchboardListener");
 DEFINE_LOG_CATEGORY(LogSwitchboard);
 
 
-int32 InitEngine(const FString& InCommandLine)
+int32 InitEngine(const TCHAR* InCommandLine)
 {
-	const int32 InitResult = GEngineLoop.PreInit((TEXT("SwitchboardListener %s"), *InCommandLine));
+	const int32 InitResult = GEngineLoop.PreInit(InCommandLine);
 	if (InitResult != 0)
 	{
 		return InitResult;
@@ -169,59 +171,9 @@
 	return true;
 }
 
-bool RunSwitchboardListener(const FSwitchboardCommandLineOptions& Options)
+int32 RunSwitchboardListener(const TCHAR* CommandLine)
 {
-	FSwitchboardListener Listener(Options);
+	const FSwitchboardCommandLineOptions Options = FSwitchboardCommandLineOptions::FromString(CommandLine);
 
-	if (!Listener.Init())
-	{
-		RequestEngineExit(TEXT("FSwitchboardListener init failure"));
-		return false;
-	}
-
-	double LastTime = FPlatformTime::Seconds();
-	const float IdealFrameTime = 1.0f / 30.0f;
-
-	bool bListenerIsRunning = true;
-
-	while (bListenerIsRunning)
-	{
-		const double CurrentTime = FPlatformTime::Seconds();
-		const double DeltaTime = CurrentTime - LastTime;
-		LastTime = CurrentTime;
-
-		if (DeltaTime > 0.1)
-		{
-			UE_LOG(LogSwitchboard, Warning, TEXT("Hitch detected; %.3f seconds since prior tick"), DeltaTime);
-		}
-
-		FTaskGraphInterface::Get().ProcessThreadUntilIdle(ENamedThreads::GameThread);
-
-		// Pump & Tick objects
-		FTSTicker::GetCoreTicker().Tick(DeltaTime);
-
-		bListenerIsRunning = Listener.Tick();
-
-		GFrameCounter++;
-		FStats::AdvanceFrame(false);
-		GLog->FlushThreadedLogs();
-
-		// Run garbage collection for the UObjects for the rest of the frame or at least to 2 ms
-		IncrementalPurgeGarbage(true, FMath::Max<float>(0.002f, IdealFrameTime - (FPlatformTime::Seconds() - CurrentTime)));
-
-		// Throttle main thread main fps by sleeping if we still have time
-		FPlatformProcess::Sleep(FMath::Max<float>(0.0f, IdealFrameTime - (FPlatformTime::Seconds() - CurrentTime)));
-	}
-
-	return true;
-}
-
-INT32_MAIN_INT32_ARGC_TCHAR_ARGV()
-{
-	const FString CommandLine = FCommandLine::BuildFromArgV(nullptr, ArgC, ArgV, nullptr);
-	FCommandLine::Set(*CommandLine);
-
-	const FSwitchboardCommandLineOptions Options = FSwitchboardCommandLineOptions::FromString(*CommandLine);
-
 	if (Options.OutputVersion)
 	{
@@ -292,6 +244,118 @@
 	}
 #endif
 
-	const bool bListenerResult = RunSwitchboardListener(Options);
-	return bListenerResult ? 0 : 1;
+	FSwitchboardListener Listener(Options);
+
+	if (!Listener.Init())
+	{
+		RequestEngineExit(TEXT("FSwitchboardListener init failure"));
+		return 1;
+	}
+
+	double LastTime = FPlatformTime::Seconds();
+	const float IdealFrameTime = 1.0f / 30.0f;
+
+	bool bListenerIsRunning = true;
+
+	while (bListenerIsRunning)
+	{
+		const double CurrentTime = FPlatformTime::Seconds();
+		const double DeltaTime = CurrentTime - LastTime;
+		LastTime = CurrentTime;
+
+		if (DeltaTime > 0.1)
+		{
+			UE_LOG(LogSwitchboard, Warning, TEXT("Hitch detected; %.3f seconds since prior tick"), DeltaTime);
+		}
+
+		FTaskGraphInterface::Get().ProcessThreadUntilIdle(ENamedThreads::GameThread);
+
+		// Pump & Tick objects
+		FTSTicker::GetCoreTicker().Tick(DeltaTime);
+
+		bListenerIsRunning = Listener.Tick();
+
+		GFrameCounter++;
+		FStats::AdvanceFrame(false);
+		GLog->FlushThreadedLogs();
+
+		// Run garbage collection for the UObjects for the rest of the frame or at least to 2 ms
+		IncrementalPurgeGarbage(true, FMath::Max<float>(0.002f, IdealFrameTime - (FPlatformTime::Seconds() - CurrentTime)));
+
+		// Throttle main thread main fps by sleeping if we still have time
+		FPlatformProcess::Sleep(FMath::Max<float>(0.0f, IdealFrameTime - (FPlatformTime::Seconds() - CurrentTime)));
+	}
+
+	// Cleanup. See also: FEngineLoop::AppPreExit and FEngineLoop::AppExit
+	UE_LOG(LogExit, Log, TEXT("Preparing to exit."));
+	FCoreDelegates::OnPreExit.Broadcast();
+	FCoreDelegates::OnExit.Broadcast();
+
+	if (GThreadPool != nullptr)
+	{
+		GThreadPool->Destroy();
+	}
+	if (GBackgroundPriorityThreadPool != nullptr)
+	{
+		GBackgroundPriorityThreadPool->Destroy();
+	}
+	if (GIOThreadPool != nullptr)
+	{
+		GIOThreadPool->Destroy();
+	}
+
+	FTaskGraphInterface::Shutdown();
+	FPlatformMisc::PlatformTearDown();
+
+	if (GLog)
+	{
+		GLog->TearDown();
+	}
+
+	FTextLocalizationManager::TearDown();
+	FInternationalization::TearDown();
+
+	FTraceAuxiliary::Shutdown();
+
+	return 0;
 }
+
+INT32_MAIN_INT32_ARGC_TCHAR_ARGV()
+{
+	const FString CommandLine = FCommandLine::BuildFromArgV(nullptr, ArgC, ArgV, nullptr);
+	FCommandLine::Set(*CommandLine);
+
+	int32 ExitCode;
+	if (FPlatformMisc::IsDebuggerPresent())
+	{
+		ExitCode = RunSwitchboardListener(*CommandLine);
+	}
+	else
+	{
+#if PLATFORM_WINDOWS && !PLATFORM_SEH_EXCEPTIONS_DISABLED
+		__try
+#endif
+		{
+			// SetCrashHandler(nullptr) sets up default behavior for Linux and Mac interfacing with CrashReportClient
+			FPlatformMisc::SetCrashHandler(nullptr);
+			GIsGuarded = true;
+			ExitCode = RunSwitchboardListener(*CommandLine);
+			GIsGuarded = false;
+		}
+#if PLATFORM_WINDOWS && !PLATFORM_SEH_EXCEPTIONS_DISABLED
+		__except (ReportCrash(GetExceptionInformation()))
+		{
+			if (GError)
+			{
+				GError->HandleError();
+			}
+
+			// RequestExit(bForce==True) terminates this process with a specific exit code and does not return.
+			FPlatformMisc::RequestExit(true);
+			ExitCode = 1; // Suppress compiler warning.
+		}
+#endif
+	}
+
+	return ExitCode;
+}
diff --git a/Engine/Source/Programs/SwitchboardListener/Private/SwitchboardListenerVersion.h b/Engine/Source/Programs/SwitchboardListener/Private/SwitchboardListenerVersion.h
--- a/Engine/Source/Programs/SwitchboardListener/Private/SwitchboardListenerVersion.h
+++ b/Engine/Source/Programs/SwitchboardListener/Private/SwitchboardListenerVersion.h
@@ -2,6 +2,6 @@
 
 #pragma once
 
-#define SBLISTENER_VERSION_MAJOR 1
-#define SBLISTENER_VERSION_MINOR 7
+#define SBLISTENER_VERSION_MAJOR 2
+#define SBLISTENER_VERSION_MINOR 0
 #define SBLISTENER_VERSION_PATCH 0
