Update missed references of UE4 to Unreal (-ue4exe, UE4Editor, UE4Game, UE4Server, UE4Client)

#rnx

diff --git a/Engine/Source/Editor/TurnkeySupport/Private/TurnkeySupportModule.cpp b/Engine/Source/Editor/TurnkeySupport/Private/TurnkeySupportModule.cpp
--- a/Engine/Source/Editor/TurnkeySupport/Private/TurnkeySupportModule.cpp
+++ b/Engine/Source/Editor/TurnkeySupport/Private/TurnkeySupportModule.cpp
@@ -323,7 +323,7 @@ public:
 			BuildCookRunParams += FString::Printf(TEXT(" -target=%s"), *BuildTargetInfo->Name);
 		}
 
-		// let the editor add options (-ue4exe in particular)
+		// let the editor add options (-unrealexe in particular)
 		{
 			BuildCookRunParams += FString::Printf(TEXT(" %s"), *FTurnkeyEditorSupport::GetUATOptions());
 		}
diff --git a/Engine/Source/Programs/SwitchboardListener/Private/SwitchboardListener.cpp b/Engine/Source/Programs/SwitchboardListener/Private/SwitchboardListener.cpp
--- a/Engine/Source/Programs/SwitchboardListener/Private/SwitchboardListener.cpp
+++ b/Engine/Source/Programs/SwitchboardListener/Private/SwitchboardListener.cpp
@@ -1631,7 +1631,7 @@ static void FillOutFlipMode(FSyncStatus& SyncStatus, FRunningProcess* FlipModeMo
 	// TimeInSeconds,MsBetweenPresents,MsBetweenDisplayChange,MsInPresentAPI,MsUntilRenderComplete,MsUntilDisplayed
 	//
 	// e.g.
-	//   "UE4Editor.exe,10916,0x0000022096A0F830,DXGI,0,512,0,Composed: Flip,1,3.753577,22.845,0.000,0.880,0.946,0.000"
+	//   "UnrealEditor.exe,10916,0x0000022096A0F830,DXGI,0,512,0,Composed: Flip,1,3.753577,22.845,0.000,0.880,0.946,0.000"
 
 	TArray<FString> Fields;
 
