SwitchboardListener: Fix for building with bForceEnableExceptions (error C2712: cannot use __try in functions that require object unwinding).

diff --git a/Engine/Source/Programs/SwitchboardListener/Private/SwitchboardListenerMain.cpp b/Engine/Source/Programs/SwitchboardListener/Private/SwitchboardListenerMain.cpp
--- a/Engine/Source/Programs/SwitchboardListener/Private/SwitchboardListenerMain.cpp
+++ b/Engine/Source/Programs/SwitchboardListener/Private/SwitchboardListenerMain.cpp
@@ -171,9 +171,12 @@ bool HandleRedeploy(FProcHandle& RedeployParentProc, uint32 RedeployParentPid)
 	return true;
 }
 
-int32 RunSwitchboardListener(const TCHAR* CommandLine)
+int32 RunSwitchboardListener(int32 ArgC, TCHAR* ArgV[])
 {
-	const FSwitchboardCommandLineOptions Options = FSwitchboardCommandLineOptions::FromString(CommandLine);
+	const FString CommandLine = FCommandLine::BuildFromArgV(nullptr, ArgC, ArgV, nullptr);
+	FCommandLine::Set(*CommandLine);
+
+	const FSwitchboardCommandLineOptions Options = FSwitchboardCommandLineOptions::FromString(*CommandLine);
 
 	if (Options.OutputVersion)
 	{
@@ -208,7 +211,7 @@ int32 RunSwitchboardListener(const TCHAR* CommandLine)
 	}
 #endif
 
-	const int32 InitResult = InitEngine(CommandLine);
+	const int32 InitResult = InitEngine(*CommandLine);
 
 	UE_LOG(LogSwitchboard, Display, TEXT("SwitchboardListener %u.%u.%u"), SBLISTENER_VERSION_MAJOR, SBLISTENER_VERSION_MINOR, SBLISTENER_VERSION_PATCH);
 
@@ -322,13 +325,10 @@ int32 RunSwitchboardListener(const TCHAR* CommandLine)
 
 INT32_MAIN_INT32_ARGC_TCHAR_ARGV()
 {
-	const FString CommandLine = FCommandLine::BuildFromArgV(nullptr, ArgC, ArgV, nullptr);
-	FCommandLine::Set(*CommandLine);
-
 	int32 ExitCode;
 	if (FPlatformMisc::IsDebuggerPresent())
 	{
-		ExitCode = RunSwitchboardListener(*CommandLine);
+		ExitCode = RunSwitchboardListener(ArgC, ArgV);
 	}
 	else
 	{
@@ -339,7 +339,7 @@ INT32_MAIN_INT32_ARGC_TCHAR_ARGV()
 			// SetCrashHandler(nullptr) sets up default behavior for Linux and Mac interfacing with CrashReportClient
 			FPlatformMisc::SetCrashHandler(nullptr);
 			GIsGuarded = true;
-			ExitCode = RunSwitchboardListener(*CommandLine);
+			ExitCode = RunSwitchboardListener(ArgC, ArgV);
 			GIsGuarded = false;
 		}
 #if PLATFORM_WINDOWS && !PLATFORM_SEH_EXCEPTIONS_DISABLED
