PR #9019: Fixed Switchboard Listener windows CPU monitor not working on non-english systems (Contributed by [REDACTED_PERSON])

diff --git a/Engine/Source/Programs/SwitchboardListener/Private/Windows/WindowsCpuUtilizationMonitor.cpp b/Engine/Source/Programs/SwitchboardListener/Private/Windows/WindowsCpuUtilizationMonitor.cpp
--- a/Engine/Source/Programs/SwitchboardListener/Private/Windows/WindowsCpuUtilizationMonitor.cpp
+++ b/Engine/Source/Programs/SwitchboardListener/Private/Windows/WindowsCpuUtilizationMonitor.cpp
@@ -22,10 +22,10 @@ FWindowsCpuUtilizationMonitor::FWindowsCpuUtilizationMonitor()
 	for (int32 CoreIdx = 0; CoreIdx < NumCores; ++CoreIdx)
 	{
 		const FString CounterName = FString::Printf(TEXT("\\Processor(%d)\\%% Processor Time"), CoreIdx);
-		PdhStatus = ::PdhAddCounter(QueryHandle, *CounterName, 0, &CounterHandles[CoreIdx]);
+		PdhStatus = ::PdhAddEnglishCounter(QueryHandle, *CounterName, 0, &CounterHandles[CoreIdx]);
 		if (PdhStatus != ERROR_SUCCESS)
 		{
-			UE_LOG(LogCpuUtilizationMonitor, Warning, TEXT("PdhAddCounter failed. Error code: %d"), PdhStatus);
+			UE_LOG(LogCpuUtilizationMonitor, Warning, TEXT("PdhAddEnglishCounter failed. Error code: %d"), PdhStatus);
 			return;
 		}
 	}
