Mitigate issues with syncing PCBs by renaming/moving existing (active) listener binary on disk.

TODO: Need to restore the executable in case it wasn't part of the PCB bundle.

diff --git a/Engine/Plugins/VirtualProduction/Switchboard/Source/Switchboard/switchboard/devices/unreal/plugin_unreal.py b/Engine/Plugins/VirtualProduction/Switchboard/Source/Switchboard/switchboard/devices/unreal/plugin_unreal.py
--- a/Engine/Plugins/VirtualProduction/Switchboard/Source/Switchboard/switchboard/devices/unreal/plugin_unreal.py
+++ b/Engine/Plugins/VirtualProduction/Switchboard/Source/Switchboard/switchboard/devices/unreal/plugin_unreal.py
@@ -1093,6 +1093,11 @@ class DeviceUnreal(Device):
             if self.unrealgamesync_lib_dir_setting:
                 sync_args += f' --ugs-lib-dir={self.unrealgamesync_lib_dir_setting}'
 
+            # If we're syncing 'Precompiled Binaries' one of those binaries may be the SwitchboardListener executable
+            # which means (on Windows atleast) we need to move the executable to make way for the new one
+            _, msg = message_protocol.create_free_listener_bin_message()
+            self.unreal_client.send_message(msg)
+
         puuid, msg = message_protocol.create_start_process_message(
             prog_path=sync_tool,
             prog_args=sync_args,
diff --git a/Engine/Plugins/VirtualProduction/Switchboard/Source/Switchboard/switchboard/message_protocol.py b/Engine/Plugins/VirtualProduction/Switchboard/Source/Switchboard/switchboard/message_protocol.py
--- a/Engine/Plugins/VirtualProduction/Switchboard/Source/Switchboard/switchboard/message_protocol.py
+++ b/Engine/Plugins/VirtualProduction/Switchboard/Source/Switchboard/switchboard/message_protocol.py
@@ -105,6 +105,13 @@ def create_redeploy_listener_message(base64listener: str, sha1digest: str):
     message = json.dumps(redeploy_cmd).encode() + b'\x00'
     return (cmd_id, message)
 
+def create_free_listener_bin_message():
+    ''' Sends a command to the listener to move its executable. '''
+    cmd_id = uuid.uuid4()
+    rename_proc_cmd = {'command': 'free binary', 'id': str(cmd_id)}
+    message = json.dumps(rename_proc_cmd).encode() + b'\x00'
+    return (cmd_id, message)
+
 def create_fixExeFlags_message(puuid):
     cmd_id = uuid.uuid4()
     cmd = {'command': 'fixExeFlags', 'id': str(cmd_id), 'uuid': str(puuid)}
diff --git a/Engine/Source/Programs/SwitchboardListener/Private/SwitchboardListener.cpp b/Engine/Source/Programs/SwitchboardListener/Private/SwitchboardListener.cpp
--- a/Engine/Source/Programs/SwitchboardListener/Private/SwitchboardListener.cpp
+++ b/Engine/Source/Programs/SwitchboardListener/Private/SwitchboardListener.cpp
@@ -494,6 +494,11 @@ bool FSwitchboardListener::RunScheduledTask(const FSwitchboardTask& InTask)
 			const FSwitchboardSetInactiveTimeoutTask& Task = static_cast<const FSwitchboardSetInactiveTimeoutTask&>(InTask);
 			return Task_SetInactiveTimeout(Task);
 		}
+		case ESwitchboardTaskType::FreeListenerBinary:
+		{
+			const FSwitchboardFreeListenerBinaryTask& Task = static_cast<const FSwitchboardFreeListenerBinaryTask&>(InTask);
+			return Task_FreeListenerBinary(Task);
+		}
 		default:
 		{
 			static const FString Response = TEXT("Unknown Command detected");
@@ -2041,6 +2046,31 @@ bool FSwitchboardListener::Task_SetInactiveTimeout(const FSwitchboardSetInactive
 	return true;
 }
 
+bool FSwitchboardListener::Task_FreeListenerBinary(const FSwitchboardFreeListenerBinaryTask& InFreeListenerBinaryTask)
+{
+	SWITCHBOARD_TRACE_CPUPROFILER_EVENT_SCOPE(FSwitchboardListener::Task_FreeListenerBinary);
+
+	const uint32 CurrentPid = FPlatformProcess::GetCurrentProcessId();
+
+	// NOTE: FPlatformProcess::ExecutablePath() is stale if we were moved while running.
+	const FString OriginalThisExePath = FPlatformProcess::GetApplicationName(CurrentPid);
+	const FString ThisExeDir = FPaths::GetPath(OriginalThisExePath);
+	const FString ThisExeFilename = FPaths::GetCleanFilename(OriginalThisExePath);
+	
+	const FString MovedListenerPath = FPaths::EngineIntermediateDir() / TEXT("Switchboard") / TEXT("old_") + ThisExeFilename;
+
+	if (!IFileManager::Get().Move(*MovedListenerPath, *OriginalThisExePath, true))
+	{
+		const uint32 LastError = FPlatformMisc::GetLastError();
+		const FString ErrorMsg = FString::Printf(TEXT("Unable to move listener exe to \"%s\" (error code %u)"), *MovedListenerPath, LastError);
+		UE_LOG(LogSwitchboard, Error, TEXT("Free listener binary: %s"), *ErrorMsg);
+		SendMessage(CreateTaskDeclinedMessage(InFreeListenerBinaryTask, ErrorMsg, {}), InFreeListenerBinaryTask.Recipient);
+		return false;
+	}
+
+	return true;
+}
+
 void FSwitchboardListener::CleanUpDisconnectedSockets()
 {
 	SWITCHBOARD_TRACE_CPUPROFILER_EVENT_SCOPE(FSwitchboardListener::CleanUpDisconnectedSockets);
diff --git a/Engine/Source/Programs/SwitchboardListener/Private/SwitchboardListener.h b/Engine/Source/Programs/SwitchboardListener/Private/SwitchboardListener.h
--- a/Engine/Source/Programs/SwitchboardListener/Private/SwitchboardListener.h
+++ b/Engine/Source/Programs/SwitchboardListener/Private/SwitchboardListener.h
@@ -25,6 +25,7 @@ struct FSwitchboardRedeployListenerTask;
 struct FSwitchboardFixExeFlagsTask;
 struct FSwitchboardMinimizeWindowsTask;
 struct FSwitchboardSetInactiveTimeoutTask;
+struct FSwitchboardFreeListenerBinaryTask;
 
 class FInternetAddr;
 class FSocket;
@@ -128,6 +129,7 @@ private:
 	bool Task_FixExeFlags(const FSwitchboardFixExeFlagsTask& InFixExeFlagsTask);
 	bool Task_MinimizeWindows(const FSwitchboardMinimizeWindowsTask& InRefreshMosaicsTask);
 	bool Task_SetInactiveTimeout(const FSwitchboardSetInactiveTimeoutTask& InTimeoutTask);
+	bool Task_FreeListenerBinary(const FSwitchboardFreeListenerBinaryTask& InRenameProcessTask);
 
 	bool KillProcessNow(FRunningProcess* InProcess, float SoftKillTimeout = 0.0f);
 	FRunningProcess* FindOrStartFlipModeMonitorForUUID(const FGuid& UUID);
diff --git a/Engine/Source/Programs/SwitchboardListener/Private/SwitchboardProtocol.cpp b/Engine/Source/Programs/SwitchboardListener/Private/SwitchboardProtocol.cpp
--- a/Engine/Source/Programs/SwitchboardListener/Private/SwitchboardProtocol.cpp
+++ b/Engine/Source/Programs/SwitchboardListener/Private/SwitchboardProtocol.cpp
@@ -252,7 +252,7 @@ bool CreateTaskFromCommand(const FString& InCommand, const FIPv4Endpoint& InEndp
 		}
 	}
 	else if (CommandName == TEXT("redeploy listener"))
-	{
+	{		
 		TSharedPtr<FJsonValue> Sha1Field = TryGetCommandRequiredField(JsonData, TEXT("sha1"));
 		TSharedPtr<FJsonValue> FileContentField = TryGetCommandRequiredField(JsonData, TEXT("content"));
 
@@ -262,6 +262,11 @@ bool CreateTaskFromCommand(const FString& InCommand, const FIPv4Endpoint& InEndp
 			return true;
 		}
 	}
+	else if (CommandName == TEXT("free binary"))
+	{
+		OutTask = MakeUnique<FSwitchboardFreeListenerBinaryTask>(MessageID, InEndpoint);
+		return true;
+	}
 	else if (CommandName == TEXT("fixExeFlags"))
 	{
 		TSharedPtr<FJsonValue> UUIDField = TryGetCommandRequiredField(JsonData, TEXT("uuid"));
diff --git a/Engine/Source/Programs/SwitchboardListener/Private/SwitchboardTasks.h b/Engine/Source/Programs/SwitchboardListener/Private/SwitchboardTasks.h
--- a/Engine/Source/Programs/SwitchboardListener/Private/SwitchboardTasks.h
+++ b/Engine/Source/Programs/SwitchboardListener/Private/SwitchboardTasks.h
@@ -22,6 +22,7 @@ enum class ESwitchboardTaskType : uint8
 	RefreshMosaics,
 	MinimizeWindows,
 	SetInactiveTimeout,
+	FreeListenerBinary
 };
 
 
@@ -250,3 +251,10 @@ struct FSwitchboardSetInactiveTimeoutTask : public FSwitchboardTask
 	//~ End FSwitchboardTask interface
 };
 
+struct FSwitchboardFreeListenerBinaryTask : public FSwitchboardTask
+{
+	FSwitchboardFreeListenerBinaryTask(const FGuid& InTaskID, const FIPv4Endpoint& InEndpoint)
+		: FSwitchboardTask{ ESwitchboardTaskType::FreeListenerBinary, TEXT("frees (moves) the listener binary file on disk so that it can be overwritten"), InTaskID, InEndpoint }
+	{}
+};
+
