Accounting for issue where the SwitchboardListener exe may not be in the PCB bundle... just unlocking and copying back the running executable (to make dumb Windows happy).

diff --git a/Engine/Source/Programs/SwitchboardListener/Private/SwitchboardListener.cpp b/Engine/Source/Programs/SwitchboardListener/Private/SwitchboardListener.cpp
--- a/Engine/Source/Programs/SwitchboardListener/Private/SwitchboardListener.cpp
+++ b/Engine/Source/Programs/SwitchboardListener/Private/SwitchboardListener.cpp
@@ -2059,6 +2059,10 @@ bool FSwitchboardListener::Task_FreeListenerBinary(const FSwitchboardFreeListene
 	
 	const FString MovedListenerPath = FPaths::EngineIntermediateDir() / TEXT("Switchboard") / TEXT("old_") + ThisExeFilename;
 
+	//
+	// Weird hackery to get around Windows locking down an executable file on disk whilst running.
+	// We move the (locked) file (which is allowed), and then make a copy (which is no longer linked to the running process)
+
 	if (!IFileManager::Get().Move(*MovedListenerPath, *OriginalThisExePath, true))
 	{
 		const uint32 LastError = FPlatformMisc::GetLastError();
@@ -2068,6 +2072,15 @@ bool FSwitchboardListener::Task_FreeListenerBinary(const FSwitchboardFreeListene
 		return false;
 	}
 
+	if (IFileManager::Get().Copy(*OriginalThisExePath, *MovedListenerPath, true, true) != ECopyResult::COPY_OK)
+	{
+		const uint32 LastError = FPlatformMisc::GetLastError();
+		const FString ErrorMsg = FString::Printf(TEXT("Unable to copy listener exe back to \"%s\" (error code %u)"), *OriginalThisExePath, LastError);
+		UE_LOG(LogSwitchboard, Error, TEXT("Free listener binary: %s"), *ErrorMsg);
+		SendMessage(CreateTaskDeclinedMessage(InFreeListenerBinaryTask, ErrorMsg, {}), InFreeListenerBinaryTask.Recipient);
+		return false;
+	}
+
 	return true;
 }
 
