SwitchboardListener: Fix expansion of %TEMP% on non-Windows platforms, affected nDisplay launch on Linux. Bump Listener version to 2.1.1.

diff --git a/Engine/Source/Programs/SwitchboardListener/Private/SwitchboardListener.cpp b/Engine/Source/Programs/SwitchboardListener/Private/SwitchboardListener.cpp
--- a/Engine/Source/Programs/SwitchboardListener/Private/SwitchboardListener.cpp
+++ b/Engine/Source/Programs/SwitchboardListener/Private/SwitchboardListener.cpp
@@ -940,8 +940,7 @@ bool FSwitchboardListener::Task_ReceiveFileFromClient(const FSwitchboardReceiveF
 
 	if (Destination.Contains(TEXT("%TEMP%")))
 	{
-		const FString TempDir = FPlatformMisc::GetEnvironmentVariable(TEXT("TEMP"));
-		Destination.ReplaceInline(TEXT("%TEMP%"), *TempDir);
+		Destination.ReplaceInline(TEXT("%TEMP%"), FPlatformProcess::UserTempDir());
 	}
 	if (Destination.Contains(TEXT("%RANDOM%")))
 	{
diff --git a/Engine/Source/Programs/SwitchboardListener/Private/SwitchboardListenerVersion.h b/Engine/Source/Programs/SwitchboardListener/Private/SwitchboardListenerVersion.h
--- a/Engine/Source/Programs/SwitchboardListener/Private/SwitchboardListenerVersion.h
+++ b/Engine/Source/Programs/SwitchboardListener/Private/SwitchboardListenerVersion.h
@@ -4,4 +4,4 @@
 
 #define SBLISTENER_VERSION_MAJOR 2
 #define SBLISTENER_VERSION_MINOR 2
-#define SBLISTENER_VERSION_PATCH 0
+#define SBLISTENER_VERSION_PATCH 1
