Switchboard: Updated to PresentMon 1.8.0 that adds support for Win11.

diff --git a/Engine/Source/Programs/SwitchboardListener/Private/SwitchboardListener.cpp b/Engine/Source/Programs/SwitchboardListener/Private/SwitchboardListener.cpp
--- a/Engine/Source/Programs/SwitchboardListener/Private/SwitchboardListener.cpp
+++ b/Engine/Source/Programs/SwitchboardListener/Private/SwitchboardListener.cpp
@@ -1555,7 +1555,7 @@ FRunningProcess* FSwitchboardListener::FindOrStartFlipModeMonitorForUUID(const F
 	const int32 PriorityModifier = 0;
 	const TCHAR* WorkingDirectory = nullptr;
 
-	MonitorProcess->Path = FPaths::EngineDir() / TEXT("Binaries") / TEXT("ThirdParty") / TEXT("PresentMon") / TEXT("Win64") / TEXT("PresentMon64-1.5.2.exe");
+	MonitorProcess->Path = FPaths::EngineDir() / TEXT("Binaries") / TEXT("ThirdParty") / TEXT("PresentMon") / TEXT("Win64") / TEXT("PresentMon-1.8.0-x64.exe");
 
 	FString Arguments = 
 		FString::Printf(TEXT("-session_name session_%d -output_stdout -dont_restart_as_admin -terminate_on_proc_exit -stop_existing_session -process_id %d"), 
@@ -1631,11 +1631,14 @@ static void FillOutFlipMode(FSyncStatus& SyncStatus, FRunningProcess* FlipModeMo
 
 	// Interpret the output as follows:
 	//
-	// Application,ProcessID,SwapChainAddress,Runtime,SyncInterval,PresentFlags,AllowsTearing,PresentMode,Dropped,
-	// TimeInSeconds,MsBetweenPresents,MsBetweenDisplayChange,MsInPresentAPI,MsUntilRenderComplete,MsUntilDisplayed
+	// Application,ProcessID,SwapChainAddress,Runtime,
+	// SyncInterval,PresentFlags,AllowsTearing,
+	// TimeInSeconds,MsBetweenPresents,MsBetweenDisplayChange,Dropped,
+	// PresentMode,
+	// MsInPresentAPI,MsUntilRenderComplete,MsUntilDisplayed
 	//
 	// e.g.
-	//   "UnrealEditor.exe,10916,0x0000022096A0F830,DXGI,0,512,0,Composed: Flip,1,3.753577,22.845,0.000,0.880,0.946,0.000"
+	//   "UnrealEditor.exe,23256,0x000002DFBFE603A0,DXGI,1,0  ,0,35.65057220000000,1.91800000000000,30735.98660000000018,0,Composed: Flip,1.91820000000000,32.67810000000000,30700.80250000000160"
 
 	TArray<FString> Fields;
 
@@ -1648,7 +1651,7 @@ static void FillOutFlipMode(FSyncStatus& SyncStatus, FRunningProcess* FlipModeMo
 			continue;
 		}
 
-		const int32 PresentMonIdx = 7;
+		const int32 PresentMonIdx = 11;
 
 		SyncStatus.FlipModeHistory.Add(Fields[PresentMonIdx]); // The first one will be "PresentMode". This is ok. 
 	}
diff --git a/Engine/Source/Programs/SwitchboardListener/SwitchboardListener.Build.cs b/Engine/Source/Programs/SwitchboardListener/SwitchboardListener.Build.cs
--- a/Engine/Source/Programs/SwitchboardListener/SwitchboardListener.Build.cs
+++ b/Engine/Source/Programs/SwitchboardListener/SwitchboardListener.Build.cs
@@ -30,7 +30,7 @@ public class SwitchboardListener : ModuleRules
 			PublicSystemLibraries.Add("Pdh.lib");
 
 			// Add PresentMon as a Runtime Dependency to avoid it from being filtered out when generating the release binaries.
-			RuntimeDependencies.Add("$(EngineDir)/Binaries/ThirdParty/PresentMon/Win64/PresentMon64-1.5.2.exe");
+			RuntimeDependencies.Add("$(EngineDir)/Binaries/ThirdParty/PresentMon/Win64/PresentMon-1.8.0-x64.exe");
 		}
 
 		PrivateIncludePathModuleNames.AddRange(
