SwitchboardListener: Fixed returning config path with double slash.

This double slash doesn't work in file explorer when copy pasting, so the user has to edit it out every time.

diff --git a/Engine/Source/Programs/SwitchboardListener/Private/SwitchboardListener.cpp b/Engine/Source/Programs/SwitchboardListener/Private/SwitchboardListener.cpp
--- a/Engine/Source/Programs/SwitchboardListener/Private/SwitchboardListener.cpp
+++ b/Engine/Source/Programs/SwitchboardListener/Private/SwitchboardListener.cpp
@@ -940,7 +940,14 @@ bool FSwitchboardListener::Task_ReceiveFileFromClient(const FSwitchboardReceiveF
 
 	if (Destination.Contains(TEXT("%TEMP%")))
 	{
-		Destination.ReplaceInline(TEXT("%TEMP%"), FPlatformProcess::UserTempDir());
+		FString TempDir = FPlatformProcess::UserTempDir();
+
+		if (TempDir.EndsWith(TEXT("/")) || TempDir.EndsWith(TEXT("\\")))
+		{
+			TempDir.LeftChopInline(1, false);
+		}
+
+		Destination.ReplaceInline(TEXT("%TEMP%"), *TempDir);
 	}
 	if (Destination.Contains(TEXT("%RANDOM%")))
 	{
