SwitchboardListener: Fix for incorrectly attempting to pass the parent process read end of the stdout pipe to the child process as its stdin pipe.

Because the handle is not inheritable, this resulted in the child process having an invalid handle associated with stdin.

Other changes:
 - Fixed an analogous bug in SwitchboardEditor that affected the venv installation script.
 - Removed workarounds from sb_setup.py and sbl_helper.py that are no longer necessary.
 - Removed unnecessary copy constructor from FRunningProcess, and replaced incomplete default constructor with initialization of members in their declarations.
 - Removed redundant ESPMode::ThreadSafe template parameters, since this is the default in UE5.
 - Removed use of auto keyword that was not compliant with the coding standard.

diff --git a/Engine/Plugins/VirtualProduction/Switchboard/Source/Switchboard/sb_setup.py b/Engine/Plugins/VirtualProduction/Switchboard/Source/Switchboard/sb_setup.py
--- a/Engine/Plugins/VirtualProduction/Switchboard/Source/Switchboard/sb_setup.py
+++ b/Engine/Plugins/VirtualProduction/Switchboard/Source/Switchboard/sb_setup.py
@@ -3,7 +3,6 @@
 from __future__ import annotations
 
 import argparse
-import ctypes
 import importlib
 import io
 import json
@@ -288,16 +287,6 @@ class SbSetup:
 
 
 def main() -> int:
-    if sys.platform.startswith('win'):
-        # As invoked by UE, we wind up with a non-inheritable stdin.
-        # We can work around this in our own `Popen` invocations by passing
-        # `stdin=subprocess.DEVNULL`, but unfortunately lots of external code
-        # still fails (including `venv.EnvBuilder.create()`).
-        # This clears the handle process-wide and prevents `subprocess` usage
-        # in external code from failing with an exception.
-        kernel32 = ctypes.WinDLL('kernel32')
-        kernel32.SetStdHandle(subprocess.STD_INPUT_HANDLE, 0)
-
     app = SbSetup()
     result = app.run()
     logging.info(f'Finished with return code {result}')
diff --git a/Engine/Plugins/VirtualProduction/Switchboard/Source/Switchboard/sbl_helper.py b/Engine/Plugins/VirtualProduction/Switchboard/Source/Switchboard/sbl_helper.py
--- a/Engine/Plugins/VirtualProduction/Switchboard/Source/Switchboard/sbl_helper.py
+++ b/Engine/Plugins/VirtualProduction/Switchboard/Source/Switchboard/sbl_helper.py
@@ -1092,12 +1092,7 @@ class SbListenerHelper:
                       f'args={args}')
 
         try:
-            # stdin=subprocess.DEVNULL is required when launched via
-            # SwitchboardListener; otherwise this call tries to make the
-            # non-existent stdin inheritable, raising `OSError: [WinError 6]
-            # The handle is invalid`
-            with subprocess.Popen(args, stdin=subprocess.DEVNULL,
-                                  stdout=subprocess.PIPE) as proc:
+            with subprocess.Popen(args, stdout=subprocess.PIPE) as proc:
                 for line in proc.stdout:
                     logging.info(f'gpf> {line.decode().rstrip()}')
 
diff --git a/Engine/Plugins/VirtualProduction/Switchboard/Source/SwitchboardEditor/Private/SwitchboardScriptInterop.cpp b/Engine/Plugins/VirtualProduction/Switchboard/Source/SwitchboardEditor/Private/SwitchboardScriptInterop.cpp
--- a/Engine/Plugins/VirtualProduction/Switchboard/Source/SwitchboardEditor/Private/SwitchboardScriptInterop.cpp
+++ b/Engine/Plugins/VirtualProduction/Switchboard/Source/SwitchboardEditor/Private/SwitchboardScriptInterop.cpp
@@ -56,10 +56,7 @@ FSwitchboardUtilScript::FSwitchboardUtilScript(FStringView PythonVenvDir)
 
 FSwitchboardUtilScript::~FSwitchboardUtilScript()
 {
-	if (WritePipe || ReadPipe)
-	{
-		FPlatformProcess::ClosePipe(ReadPipe, WritePipe);
-	}
+	FPlatformProcess::ClosePipe(StdoutParentReadPipe, StdoutChildWritePipe);
 
 	if (ProcHandle.IsValid())
 	{
@@ -87,9 +84,9 @@ bool FSwitchboardUtilScript::Run(const FString& Args)
 	uint32 OutProcessId = 0;
 	const int32 PriorityModifier = 0;
 	const TCHAR* WorkingDirectory = nullptr;
-	FPlatformProcess::CreatePipe(ReadPipe, WritePipe);
+	FPlatformProcess::CreatePipe(StdoutParentReadPipe, StdoutChildWritePipe);
 
-	ProcHandle = FPlatformProcess::CreateProc(*PythonExe, *Args, bLaunchDetached, bLaunchHidden, bLaunchReallyHidden, &OutProcessId, PriorityModifier, WorkingDirectory, WritePipe, ReadPipe);
+	ProcHandle = FPlatformProcess::CreateProc(*PythonExe, *Args, bLaunchDetached, bLaunchHidden, bLaunchReallyHidden, &OutProcessId, PriorityModifier, WorkingDirectory, StdoutChildWritePipe);
 	return ProcHandle.IsValid();
 }
 
@@ -108,7 +105,7 @@ TOptional<int32> FSwitchboardUtilScript::PollStdoutAndReturnCode()
 	}
 
 	TArray<uint8> Output;
-	if (FPlatformProcess::ReadPipeToArray(ReadPipe, Output))
+	if (FPlatformProcess::ReadPipeToArray(StdoutParentReadPipe, Output))
 	{
 		StdoutBuf.Append(Output);
 	}
@@ -119,7 +116,7 @@ TOptional<int32> FSwitchboardUtilScript::PollStdoutAndReturnCode()
 		ReturnCode.Emplace(OutReturnCode);
 
 		// Additional read necessary to ensure we've captured all output.
-		if (FPlatformProcess::ReadPipeToArray(ReadPipe, Output))
+		if (FPlatformProcess::ReadPipeToArray(StdoutParentReadPipe, Output))
 		{
 			StdoutBuf.Append(Output);
 		}
diff --git a/Engine/Plugins/VirtualProduction/Switchboard/Source/SwitchboardEditor/Private/SwitchboardScriptInterop.h b/Engine/Plugins/VirtualProduction/Switchboard/Source/SwitchboardEditor/Private/SwitchboardScriptInterop.h
--- a/Engine/Plugins/VirtualProduction/Switchboard/Source/SwitchboardEditor/Private/SwitchboardScriptInterop.h
+++ b/Engine/Plugins/VirtualProduction/Switchboard/Source/SwitchboardEditor/Private/SwitchboardScriptInterop.h
@@ -24,8 +24,8 @@ public:
 public:
 	FString PythonExe;
 	FProcHandle ProcHandle;
-	void* WritePipe = nullptr;
-	void* ReadPipe = nullptr;
+	void* StdoutChildWritePipe = nullptr;
+	void* StdoutParentReadPipe = nullptr;
 	TArray<uint8> StdoutBuf;
 	TOptional<int32> ReturnCode;
 };
diff --git a/Engine/Source/Programs/SwitchboardListener/Private/SwitchboardListener.cpp b/Engine/Source/Programs/SwitchboardListener/Private/SwitchboardListener.cpp
--- a/Engine/Source/Programs/SwitchboardListener/Private/SwitchboardListener.cpp
+++ b/Engine/Source/Programs/SwitchboardListener/Private/SwitchboardListener.cpp
@@ -89,12 +89,15 @@ namespace
 
 struct FRunningProcess
 {
-	uint32 PID;
+	uint32 PID = 0;
 	FGuid UUID;
 	FProcHandle Handle;
 
-	void* WritePipe;
-	void* ReadPipe;
+	void* StdoutParentReadPipe = nullptr;
+	void* StdoutChildWritePipe = nullptr;
+	void* StdinChildReadPipe = nullptr;
+	void* StdinParentWritePipe = nullptr;
+
 	TArray<uint8> Output;
 
 	FIPv4Endpoint Recipient;
@@ -102,29 +105,30 @@ struct FRunningProcess
 	FString Name;
 	FString Caller;
 
-	std::atomic<bool> bPendingKill;
-	bool bUpdateClientsWithStdout;
-
-	FRunningProcess()
-		: bPendingKill(false)
-	{}
+	std::atomic<bool> bPendingKill = false;
+	bool bUpdateClientsWithStdout = false;
 
-	FRunningProcess(const FRunningProcess& InProcess)
+	bool CreatePipes()
 	{
-		PID       = InProcess.PID;
-		UUID      = InProcess.UUID;
-		Handle    = InProcess.Handle;
-		WritePipe = InProcess.WritePipe;
-		ReadPipe  = InProcess.ReadPipe;
-		Output    = InProcess.Output;
-		Recipient = InProcess.Recipient;
-		Path      = InProcess.Path;
-		Name      = InProcess.Name;
-		Caller    = InProcess.Caller;
+		if (!FPlatformProcess::CreatePipe(StdoutParentReadPipe, StdoutChildWritePipe))
+		{
+			UE_LOG(LogSwitchboard, Error, TEXT("Could not create pipe to read process output!"));
+			return false;
+		}
 
-		bUpdateClientsWithStdout = InProcess.bUpdateClientsWithStdout;
+		if (!FPlatformProcess::CreatePipe(StdinChildReadPipe, StdinParentWritePipe, true))
+		{
+			UE_LOG(LogSwitchboard, Error, TEXT("Could not create pipe to write process input!"));
+			return false;
+		}
+
+		return true;
+	}
 
-		bPendingKill.store(InProcess.bPendingKill);
+	void ClosePipes()
+	{
+		FPlatformProcess::ClosePipe(StdoutParentReadPipe, StdoutChildWritePipe);
+		FPlatformProcess::ClosePipe(StdinChildReadPipe, StdinParentWritePipe);
 	}
 };
 
@@ -212,10 +216,10 @@ FString FSwitchboardCommandLineOptions::ToString(bool bIncludeRedeploy /* = fals
 FSwitchboardListener::FSwitchboardListener(const FSwitchboardCommandLineOptions& InOptions)
 	: Options(InOptions)
 	, SocketListener(nullptr)
-	, CpuMonitor(MakeShared<FCpuUtilizationMonitor, ESPMode::ThreadSafe>())
+	, CpuMonitor(MakeShared<FCpuUtilizationMonitor>())
 	, bIsNvAPIInitialized(false)
-	, CachedMosaicToposLock(MakeShared<FRWLock, ESPMode::ThreadSafe>())
-	, CachedMosaicTopos(MakeShared<TArray<FMosaicTopo>, ESPMode::ThreadSafe>())
+	, CachedMosaicToposLock(MakeShared<FRWLock>())
+	, CachedMosaicTopos(MakeShared<TArray<FMosaicTopo>>())
 {
 #if PLATFORM_WINDOWS
 	// initialize NvAPI
@@ -309,7 +313,7 @@ bool FSwitchboardListener::Tick()
 				StatePacket.TotalPhysicalMemory = FPlatformMemory::GetConstants().TotalPhysical;
 				StatePacket.PlatformBinaryDirectory = FPlatformProcess::GetBinariesSubdirectory();
 
-				for (const auto& RunningProcess : RunningProcesses)
+				for (const TSharedPtr<FRunningProcess>& RunningProcess : RunningProcesses)
 				{
 					check(RunningProcess.IsValid());
 
@@ -635,7 +639,7 @@ bool FSwitchboardListener::Task_StartProcess(const FSwitchboardStartTask& InRunT
 {
 	SWITCHBOARD_TRACE_CPUPROFILER_EVENT_SCOPE(FSwitchboardListener::Task_StartProcess);
 
-	auto NewProcess = MakeShared<FRunningProcess, ESPMode::ThreadSafe>();
+	TSharedRef<FRunningProcess> NewProcess = MakeShared<FRunningProcess>();
 
 	NewProcess->Recipient = InRunTask.Recipient;
 	NewProcess->Path = InRunTask.Command;
@@ -645,9 +649,8 @@ bool FSwitchboardListener::Task_StartProcess(const FSwitchboardStartTask& InRunT
 	NewProcess->UUID = InRunTask.TaskID; // Process ID is the same as the message ID.
 	NewProcess->PID = 0; // default value
 
-	if (!FPlatformProcess::CreatePipe(NewProcess->ReadPipe, NewProcess->WritePipe))
+	if (!NewProcess->CreatePipes())
 	{
-		UE_LOG(LogSwitchboard, Error, TEXT("Could not create pipe to read process output!"));
 		return false;
 	}
 	
@@ -666,8 +669,8 @@ bool FSwitchboardListener::Task_StartProcess(const FSwitchboardStartTask& InRunT
 		&NewProcess->PID, 
 		PriorityModifier, 
 		WorkingDirectory, 
-		NewProcess->WritePipe,
-		NewProcess->ReadPipe
+		NewProcess->StdoutChildWritePipe,
+		NewProcess->StdinChildReadPipe
 	);
 
 	if (!NewProcess->Handle.IsValid())
@@ -676,7 +679,7 @@ bool FSwitchboardListener::Task_StartProcess(const FSwitchboardStartTask& InRunT
 		FPlatformProcess::CloseProc(NewProcess->Handle);
 
 		// close pipes
-		FPlatformProcess::ClosePipe(NewProcess->ReadPipe, NewProcess->WritePipe);
+		NewProcess->ClosePipes();
 
 		// log error
 		const FString ErrorMsg = FString::Printf(TEXT("Could not start program %s"), *InRunTask.Command);
@@ -742,9 +745,9 @@ bool FSwitchboardListener::Task_KillProcess(const FSwitchboardKillTask& KillTask
 	}
 
 	// Look in RunningProcesses
-	TSharedPtr<FRunningProcess, ESPMode::ThreadSafe> Process;
+	TSharedPtr<FRunningProcess> Process;
 	{
-		auto* ProcessPtr = RunningProcesses.FindByPredicate([&KillTask](const TSharedPtr<FRunningProcess, ESPMode::ThreadSafe>& InProcess)
+		TSharedPtr<FRunningProcess>* ProcessPtr = RunningProcesses.FindByPredicate([&KillTask](const TSharedPtr<FRunningProcess>& InProcess)
 		{
 			check(InProcess.IsValid());
 			return !InProcess->bPendingKill && (InProcess->UUID == KillTask.ProgramID);
@@ -758,9 +761,9 @@ bool FSwitchboardListener::Task_KillProcess(const FSwitchboardKillTask& KillTask
 	}
 
 	// Look in FlipModeMonitors
-	TSharedPtr<FRunningProcess, ESPMode::ThreadSafe> FlipModeMonitor;
+	TSharedPtr<FRunningProcess> FlipModeMonitor;
 	{
-		auto* FlipModeMonitorPtr = FlipModeMonitors.FindByPredicate([&](const TSharedPtr<FRunningProcess, ESPMode::ThreadSafe>& FlipMonitor)
+		TSharedPtr<FRunningProcess>* FlipModeMonitorPtr = FlipModeMonitors.FindByPredicate([&](const TSharedPtr<FRunningProcess>& FlipMonitor)
 		{
 			check(FlipMonitor.IsValid());
 			return !FlipMonitor->bPendingKill && (FlipMonitor->UUID == KillTask.ProgramID);
@@ -857,9 +860,9 @@ bool FSwitchboardListener::Task_FixExeFlags(const FSwitchboardFixExeFlagsTask& T
 	}
 
 	// Look in RunningProcesses
-	TSharedPtr<FRunningProcess, ESPMode::ThreadSafe> Process;
+	TSharedPtr<FRunningProcess> Process;
 	{
-		auto* ProcessPtr = RunningProcesses.FindByPredicate([&Task](const TSharedPtr<FRunningProcess, ESPMode::ThreadSafe>& InProcess)
+		TSharedPtr<FRunningProcess>* ProcessPtr = RunningProcesses.FindByPredicate([&Task](const TSharedPtr<FRunningProcess>& InProcess)
 			{
 				check(InProcess.IsValid());
 				return !InProcess->bPendingKill && (InProcess->UUID == Task.ProgramID);
@@ -1512,7 +1515,7 @@ FRunningProcess* FSwitchboardListener::FindOrStartFlipModeMonitorForUUID(const F
 
 	// See if the associated FlipModeMonitor is running
 	{
-		auto* FlipModeMonitorPtr = FlipModeMonitors.FindByPredicate([&](const TSharedPtr<FRunningProcess, ESPMode::ThreadSafe>& FlipMonitor)
+		TSharedPtr<FRunningProcess>* FlipModeMonitorPtr = FlipModeMonitors.FindByPredicate([&](const TSharedPtr<FRunningProcess>& FlipMonitor)
 		{
 			check(FlipMonitor.IsValid());
 			return FlipMonitor->UUID == UUID;
@@ -1526,9 +1529,9 @@ FRunningProcess* FSwitchboardListener::FindOrStartFlipModeMonitorForUUID(const F
 
 	// It wasn't in there, so let's find our target process
 
-	TSharedPtr<FRunningProcess, ESPMode::ThreadSafe> Process;
+	TSharedPtr<FRunningProcess> Process;
 	{
-		auto* ProcessPtr = RunningProcesses.FindByPredicate([&](const TSharedPtr<FRunningProcess, ESPMode::ThreadSafe>& InProcess)
+		TSharedPtr<FRunningProcess>* ProcessPtr = RunningProcesses.FindByPredicate([&](const TSharedPtr<FRunningProcess>& InProcess)
 		{
 			check(InProcess.IsValid());
 			return InProcess->UUID == UUID;
@@ -1547,10 +1550,9 @@ FRunningProcess* FSwitchboardListener::FindOrStartFlipModeMonitorForUUID(const F
 	}
 
 	// Ok, we need to create our monitor.
+	TSharedRef<FRunningProcess> MonitorProcess = MakeShared<FRunningProcess>();
 
-	auto MonitorProcess = MakeShared<FRunningProcess, ESPMode::ThreadSafe>();
-
-	if (!FPlatformProcess::CreatePipe(MonitorProcess->ReadPipe, MonitorProcess->WritePipe))
+	if (!MonitorProcess->CreatePipes())
 	{
 		UE_LOG(LogSwitchboard, Error, TEXT("Could not create pipe to read MonitorProcess output!"));
 		return nullptr;
@@ -1577,8 +1579,8 @@ FRunningProcess* FSwitchboardListener::FindOrStartFlipModeMonitorForUUID(const F
 		&MonitorProcess->PID,
 		PriorityModifier,
 		WorkingDirectory,
-		MonitorProcess->WritePipe,
-		MonitorProcess->ReadPipe
+		MonitorProcess->StdoutChildWritePipe,
+		MonitorProcess->StdinChildReadPipe
 	);
 
 	if (!MonitorProcess->Handle.IsValid() || !FPlatformProcess::IsProcRunning(MonitorProcess->Handle))
@@ -1587,7 +1589,7 @@ FRunningProcess* FSwitchboardListener::FindOrStartFlipModeMonitorForUUID(const F
 		FPlatformProcess::CloseProc(MonitorProcess->Handle);
 
 		// Close unused pipes
-		FPlatformProcess::ClosePipe(MonitorProcess->ReadPipe, MonitorProcess->WritePipe);
+		MonitorProcess->ClosePipes();
 
 		// Log error
 		const FString ErrorMsg = FString::Printf(TEXT("Could not start FlipMode monitor  %s"), *MonitorProcess->Path);
@@ -1858,14 +1860,14 @@ bool FSwitchboardListener::Task_GetSyncStatus(const FSwitchboardGetSyncStatusTas
 		return false;
 	}
 
-	TSharedRef<FSyncStatus, ESPMode::ThreadSafe> SyncStatus = MakeShared<FSyncStatus, ESPMode::ThreadSafe>(); // Smart pointer to avoid potentially bigger copy to lambda below.
+	TSharedRef<FSyncStatus> SyncStatus = MakeShared<FSyncStatus>(); // Smart pointer to avoid potentially bigger copy to lambda below.
 
 	// We need to run these on this thread to avoid threading issues.
 	FillOutFlipMode(SyncStatus.Get(), FindOrStartFlipModeMonitorForUUID(InGetSyncStatusTask.ProgramID));
 
 	// Fill out fullscreen optimization setting
 	{
-		auto* ProcessPtr = RunningProcesses.FindByPredicate([&](const TSharedPtr<FRunningProcess, ESPMode::ThreadSafe>& InProcess)
+		TSharedPtr<FRunningProcess>* ProcessPtr = RunningProcesses.FindByPredicate([&](const TSharedPtr<FRunningProcess>& InProcess)
 		{
 			check(InProcess.IsValid());
 			return !InProcess->bPendingKill && (InProcess->UUID == InGetSyncStatusTask.ProgramID);
@@ -2137,12 +2139,12 @@ void FSwitchboardListener::DisconnectClient(const FIPv4Endpoint& InClientEndpoin
 	ReceiveBuffer.Remove(InClientEndpoint);
 }
 
-void FSwitchboardListener::HandleStdout(const TSharedPtr<FRunningProcess, ESPMode::ThreadSafe>& Process)
+void FSwitchboardListener::HandleStdout(const TSharedPtr<FRunningProcess>& Process)
 {
 	SWITCHBOARD_TRACE_CPUPROFILER_EVENT_SCOPE(FSwitchboardListener::HandleStdout);
 
 	TArray<uint8> Output;
-	if (FPlatformProcess::ReadPipeToArray(Process->ReadPipe, Output))
+	if (FPlatformProcess::ReadPipeToArray(Process->StdoutParentReadPipe, Output))
 	{
 		Process->Output.Append(Output);
 	}
@@ -2168,14 +2170,14 @@ void FSwitchboardListener::HandleStdout(const TSharedPtr<FRunningProcess, ESPMod
 	}
 }
 
-void FSwitchboardListener::HandleRunningProcesses(TArray<TSharedPtr<FRunningProcess, ESPMode::ThreadSafe>>& Processes, bool bNotifyThatProgramEnded)
+void FSwitchboardListener::HandleRunningProcesses(TArray<TSharedPtr<FRunningProcess>>& Processes, bool bNotifyThatProgramEnded)
 {
 	SWITCHBOARD_TRACE_CPUPROFILER_EVENT_SCOPE(FSwitchboardListener::HandleRunningProcesses);
 
 	// Reads pipe and cleans up dead processes from the array.
 	for (auto Iter = Processes.CreateIterator(); Iter; ++Iter)
 	{
-		TSharedPtr<FRunningProcess, ESPMode::ThreadSafe> Process = *Iter;
+		TSharedPtr<FRunningProcess> Process = *Iter;
 
 		check(Process.IsValid());
 
@@ -2225,7 +2227,7 @@ void FSwitchboardListener::HandleRunningProcesses(TArray<TSharedPtr<FRunningProc
 
 					// Kill its monitor to avoid potential zombies (unless it is already pending kill)
 					{
-						auto* FlipModeMonitorPtr = FlipModeMonitors.FindByPredicate([&](const TSharedPtr<FRunningProcess, ESPMode::ThreadSafe>& FlipMonitor)
+						TSharedPtr<FRunningProcess>* FlipModeMonitorPtr = FlipModeMonitors.FindByPredicate([&](const TSharedPtr<FRunningProcess>& FlipMonitor)
 						{
 							check(FlipMonitor.IsValid());
 							return !FlipMonitor->bPendingKill && (FlipMonitor->UUID == Process->UUID);
@@ -2242,7 +2244,7 @@ void FSwitchboardListener::HandleRunningProcesses(TArray<TSharedPtr<FRunningProc
 				}
 
 				FPlatformProcess::CloseProc(Process->Handle);
-				FPlatformProcess::ClosePipe(Process->ReadPipe, Process->WritePipe);
+				Process->ClosePipes();
 
 				Iter.RemoveCurrent();
 			}
diff --git a/Engine/Source/Programs/SwitchboardListener/Private/SwitchboardListener.h b/Engine/Source/Programs/SwitchboardListener/Private/SwitchboardListener.h
--- a/Engine/Source/Programs/SwitchboardListener/Private/SwitchboardListener.h
+++ b/Engine/Source/Programs/SwitchboardListener/Private/SwitchboardListener.h
@@ -136,8 +136,8 @@ private:
 
 	void CleanUpDisconnectedSockets();
 	void DisconnectClient(const FIPv4Endpoint& InClientEndpoint);
-	void HandleStdout(const TSharedPtr<FRunningProcess, ESPMode::ThreadSafe>& Process);
-	void HandleRunningProcesses(TArray<TSharedPtr<FRunningProcess, ESPMode::ThreadSafe>>& Processes, bool bNotifyThatProgramEnded);
+	void HandleStdout(const TSharedPtr<FRunningProcess>& Process);
+	void HandleRunningProcesses(TArray<TSharedPtr<FRunningProcess>>& Processes, bool bNotifyThatProgramEnded);
 
 	bool SendMessage(const FString& InMessage, const FIPv4Endpoint& InEndpoint);
 	void SendMessageFutures();
@@ -159,15 +159,15 @@ private:
 
 	TQueue<TUniquePtr<FSwitchboardTask>, EQueueMode::Spsc> ScheduledTasks;
 	TQueue<TUniquePtr<FSwitchboardTask>, EQueueMode::Spsc> DisconnectTasks;
-	TArray<TSharedPtr<FRunningProcess, ESPMode::ThreadSafe>> RunningProcesses;
-	TArray<TSharedPtr<FRunningProcess, ESPMode::ThreadSafe>> FlipModeMonitors;
+	TArray<TSharedPtr<FRunningProcess>> RunningProcesses;
+	TArray<TSharedPtr<FRunningProcess>> FlipModeMonitors;
 	TArray<FSwitchboardMessageFuture> MessagesFutures;
-	TSharedPtr<FCpuUtilizationMonitor, ESPMode::ThreadSafe> CpuMonitor;
+	TSharedPtr<FCpuUtilizationMonitor> CpuMonitor;
 
 	bool bIsNvAPIInitialized;
 
-	TSharedPtr<FRWLock, ESPMode::ThreadSafe> CachedMosaicToposLock;
-	TSharedPtr<TArray<FMosaicTopo>, ESPMode::ThreadSafe> CachedMosaicTopos;
+	TSharedPtr<FRWLock> CachedMosaicToposLock;
+	TSharedPtr<TArray<FMosaicTopo>> CachedMosaicTopos;
 
 	FRedeployStatus RedeployStatus;
 };
