Fix  C4855: implicit capture of 'this' via '[=]' is deprecated in '/std:c++20'

#rnx

diff --git a/Engine/Source/Programs/SwitchboardListener/Private/SwitchboardListener.cpp b/Engine/Source/Programs/SwitchboardListener/Private/SwitchboardListener.cpp
--- a/Engine/Source/Programs/SwitchboardListener/Private/SwitchboardListener.cpp
+++ b/Engine/Source/Programs/SwitchboardListener/Private/SwitchboardListener.cpp
@@ -838,7 +838,7 @@ bool FSwitchboardListener::Task_KillProcess(const FSwitchboardKillTask& KillTask
 
 	const FGuid UUID = KillTask.ProgramID;
 
-	MessageFuture.Future = Async(EAsyncExecution::ThreadPool, [=]() {
+	MessageFuture.Future = Async(EAsyncExecution::ThreadPool, [this, KillTask, Process, FlipModeMonitor, UUID]() {
 		SWITCHBOARD_TRACE_CPUPROFILER_EVENT_SCOPE_STR("FSwitchboardListener::Task_KillProcess future closure");
 
 		const float SoftKillTimeout = 2.0f;
