SwitchboardListener: Workaround for OpenSSL unity build error. (Borrowed from QuicMessaging/WebSocketNetworking/DTLSHandler.)

diff --git a/Engine/Source/Programs/SwitchboardListener/Private/SwitchboardAuth.cpp b/Engine/Source/Programs/SwitchboardListener/Private/SwitchboardAuth.cpp
--- a/Engine/Source/Programs/SwitchboardListener/Private/SwitchboardAuth.cpp
+++ b/Engine/Source/Programs/SwitchboardListener/Private/SwitchboardAuth.cpp
@@ -25,9 +25,14 @@
 #include <termios.h>
 #endif
 
+UE_PUSH_MACRO("UI")
+// Workaround for "ossl_typ.h(144): error C2365: 'UI': redefinition; previous definition was 'namespace'"
+//                "ObjectMacros.h(872): note: see declaration of 'UI'"
+#define UI UI_ST
 #include <openssl/pem.h>
 #include <openssl/rand.h>
 #include <openssl/x509.h>
+UE_POP_MACRO("UI")
 
 
 namespace UE::SwitchboardListener
