SwitchboardListener: Fix bug where reception is not always disabled until current message is processed.

diff --git a/Engine/Source/Programs/SwitchboardListener/Private/SwitchboardListener.cpp b/Engine/Source/Programs/SwitchboardListener/Private/SwitchboardListener.cpp
--- a/Engine/Source/Programs/SwitchboardListener/Private/SwitchboardListener.cpp
+++ b/Engine/Source/Programs/SwitchboardListener/Private/SwitchboardListener.cpp
@@ -583,14 +583,16 @@ bool FSwitchboardListener::StartListening()
 	}
 
 	// Create/allocate a new listener object.
-	if (QUIC_FAILED(Status = QuicApi->ListenerOpen(QuicRegistration, QuicListenerThunk, this, &QuicListener))) {
+	if (QUIC_FAILED(Status = QuicApi->ListenerOpen(QuicRegistration, QuicListenerThunk, this, &QuicListener))) 
+	{
 		UE_LOGFMT(LogSwitchboard, Error, "MsQuic ListenerOpen failed with status {Status}", static_cast<int64>(Status));
 		return false;
 	}
 
 	// Starts listening for incoming connections.
 	QUIC_ADDR QuicAddr = QuicAddrFromEndpoint(*ListenerEndpoint);
-	if (QUIC_FAILED(Status = QuicApi->ListenerStart(QuicListener, &SblAlpn, 1, &QuicAddr))) {
+	if (QUIC_FAILED(Status = QuicApi->ListenerStart(QuicListener, &SblAlpn, 1, &QuicAddr))) 
+	{
 		UE_LOGFMT(LogSwitchboard, Error, "MsQuic ListenerStart failed with status {Status}", static_cast<int64>(Status));
 		return false;
 	}
@@ -1625,6 +1627,7 @@ bool FSwitchboardListener::Task_SendFileToClient(const FSwitchboardSendFileToCli
 
 
 //static
+_Function_class_(QUIC_LISTENER_CALLBACK)
 QUIC_STATUS QUIC_API FSwitchboardListener::QuicListenerThunk(HQUIC Listener, void* Context, QUIC_LISTENER_EVENT* Event)
 {
 	FSwitchboardListener* This = reinterpret_cast<FSwitchboardListener*>(Context);
@@ -1835,34 +1838,37 @@ QUIC_STATUS FSwitchboardListener::QuicStreamCallback(HQUIC Stream, QUIC_STREAM_E
 			// The value of this struct member when the callback returns is taken to mean
 			// how many bytes we actually consumed. Bytes left unconsumed will be handled
 			// after the tick re-enables receiving, upon consuming the preceding message.
-			const uint64 IncomingTotalBufferLength = Event->RECEIVE.TotalBufferLength;
 			Event->RECEIVE.TotalBufferLength = 0;
 
 			// Data was received from the peer on the stream.
 			Connection->LastActivityTime = FPlatformTime::Seconds();
 
 			UE::TUniqueLock<UE::FMutex> ReceiveLock(Connection->ReceiveLock);
-			for (uint32_t BufferNum = 0; BufferNum < Event->RECEIVE.BufferCount; ++BufferNum)
+
+			if (!Connection->bMessageComplete)
 			{
-				const QUIC_BUFFER Buffer = Event->RECEIVE.Buffers[BufferNum];
-				for (uint32 ReadIdx = 0; ReadIdx < Buffer.Length; ++ReadIdx)
+				for (uint32_t BufferNum = 0; BufferNum < Event->RECEIVE.BufferCount; ++BufferNum)
 				{
-					const uint8_t Byte = Buffer.Buffer[ReadIdx];
-					Connection->ReceiveBuffer->Add(Byte);
-					++Event->RECEIVE.TotalBufferLength;
-
-					// If this concluded the outstanding message, signal the main thread to
-					// process it, and let MsQuic keep the next message buffered internally.
-					// Partial buffer consumption implicitly suspends receive events, which
-					// we then resume after the main thread processes the previous message.
-					if (Byte == 0)
+					const QUIC_BUFFER Buffer = Event->RECEIVE.Buffers[BufferNum];
+					for (uint32 ReadIdx = 0; ReadIdx < Buffer.Length; ++ReadIdx)
 					{
-						Connection->bMessageComplete = true;
+						const uint8_t Byte = Buffer.Buffer[ReadIdx];
+						Connection->ReceiveBuffer->Add(Byte);
+						++Event->RECEIVE.TotalBufferLength;
+
+						// If this concluded the outstanding message, signal the main thread to
+						// process it, and let MsQuic keep the next message buffered internally.
+						// Partial buffer consumption implicitly suspends receive events, which
+						// we then resume after the main thread processes the previous message.
+						if (Byte == 0)
+						{
+							Connection->bMessageComplete = true;
 
-						// Break inner loop + outer loop + switch
-						goto stream_receive_outer_break;
+							// Break inner loop + outer loop + switch
+							goto stream_receive_outer_break;
+						}
 					}
-				}	
+				}
 			}
 
 stream_receive_outer_break:
diff --git a/Engine/Source/Programs/SwitchboardListener/Private/SwitchboardListener.h b/Engine/Source/Programs/SwitchboardListener/Private/SwitchboardListener.h
--- a/Engine/Source/Programs/SwitchboardListener/Private/SwitchboardListener.h
+++ b/Engine/Source/Programs/SwitchboardListener/Private/SwitchboardListener.h
@@ -177,6 +177,7 @@ private:
 
 	void RollbackRedeploy();
 
+	_Function_class_(QUIC_LISTENER_CALLBACK)
 	static QUIC_STATUS QUIC_API QuicListenerThunk(HQUIC Listener, void* Context, QUIC_LISTENER_EVENT* Event);
 	static QUIC_STATUS QUIC_API QuicConnectionThunk(HQUIC Connection, void* Context, QUIC_CONNECTION_EVENT* Event);
 	static QUIC_STATUS QUIC_API QuicStreamThunk(HQUIC Stream, void* Context, QUIC_STREAM_EVENT* Event);
diff --git a/Engine/Source/Programs/SwitchboardListener/Private/SwitchboardListenerVersion.h b/Engine/Source/Programs/SwitchboardListener/Private/SwitchboardListenerVersion.h
--- a/Engine/Source/Programs/SwitchboardListener/Private/SwitchboardListenerVersion.h
+++ b/Engine/Source/Programs/SwitchboardListener/Private/SwitchboardListenerVersion.h
@@ -4,4 +4,4 @@
 
 #define SBLISTENER_VERSION_MAJOR 3
 #define SBLISTENER_VERSION_MINOR 0
-#define SBLISTENER_VERSION_PATCH 0
+#define SBLISTENER_VERSION_PATCH 1
