SwitchboardListener: UE version now appears in the title.

diff --git a/Engine/Source/Programs/SwitchboardListener/Private/SwitchboardListenerMain.cpp b/Engine/Source/Programs/SwitchboardListener/Private/SwitchboardListenerMain.cpp
--- a/Engine/Source/Programs/SwitchboardListener/Private/SwitchboardListenerMain.cpp
+++ b/Engine/Source/Programs/SwitchboardListener/Private/SwitchboardListenerMain.cpp
@@ -7,6 +7,7 @@
 #include "HAL/ExceptionHandling.h"
 #include "Interfaces/IPv4/IPv4Endpoint.h"
 #include "Misc/ScopeExit.h"
+#include "Modules/BuildVersion.h"
 #include "MsQuicRuntimeModule.h"
 #include "ProfilingDebugging/TraceAuxiliary.h"
 #include "RequiredProgramMainCPPInclude.h"
@@ -98,14 +99,6 @@ FProcHandle CheckRedeploy(const FSwitchboardCommandLineOptions& Options)
 // Wait for the old listener parent process to exit and delete its executable.
 bool HandleRedeploy(FProcHandle& RedeployParentProc, uint32 RedeployParentPid)
 {
-#if PLATFORM_WINDOWS
-	// NOTE: FPlatformProcess::ExecutablePath() is stale if moved while running!
-	const FString CurrentExePath = FPlatformProcess::GetApplicationName(FPlatformProcess::GetCurrentProcessId());
-
-	// Otherwise, the window title remains the temporary (pre-rename) filename.
-	::SetConsoleTitle(*CurrentExePath);
-#endif
-
 	const FString& OldLauncherPath = FPlatformProcess::GetApplicationName(RedeployParentPid);
 
 	// Get handle to existing IPC semaphore.
@@ -325,9 +318,24 @@ int32 RunSwitchboardListener()
 	return 0;
 }
 
+// Sets the window title. It informs the version of the listener and the UE version it is based on.
+static void SetWindowTitle()
+{
+#if PLATFORM_WINDOWS
+
+	const FString Title = *FString::Printf(TEXT("Switchboard Listener %d.%d.%d based on Unreal Engine %s"),
+		SBLISTENER_VERSION_MAJOR, SBLISTENER_VERSION_MINOR, SBLISTENER_VERSION_PATCH, *FEngineVersion::Current().ToString());
+
+	::SetConsoleTitle(*Title);
+
+#endif
+
+}
 
 int32 SwitchboardListenerMain()
 {
+	SetWindowTitle();
+
 	int32 ExitCode;
 	if (FPlatformMisc::IsDebuggerPresent())
 	{
