C4702 warning fix
#rnx

diff --git a/Engine/Source/Programs/SwitchboardListener/Private/SwitchboardListener.cpp b/Engine/Source/Programs/SwitchboardListener/Private/SwitchboardListener.cpp
--- a/Engine/Source/Programs/SwitchboardListener/Private/SwitchboardListener.cpp
+++ b/Engine/Source/Programs/SwitchboardListener/Private/SwitchboardListener.cpp
@@ -821,7 +821,6 @@ bool FSwitchboardListener::RunScheduledTask(const FSwitchboardTask& InTask)
 			return false;
 		}
 	}
-	return false;
 }
 
 static uint32 FindPidInFocus()
