SwitchboardListener: Convert SwitchboardListener to a Slate application.

This provides a more user-friendly workflow for first time setup, and enables users to easily change the authentication password.

This CL also introduces a sibling SwitchboardListenerCmd target, so we can continue to support building as a console application.

Other changes:
- Factor out common startup/shutdown code from SwitchboardListenerMain.cpp to SwitchboardListenerMainCommon.{h,cpp}.
- Correctly handle "set inactive timeout" command from unauthenticated clients.
- Change user-facing "(authentication) token" verbiage to "password."
- If we can't load the private key password credential, quietly regenerate the private key and certificate pair for the user.
- Make Resource.rc2 more consistent with Default.rc2, but use SBL version numbers for FILEVERSION instead of UE version numbers.
- Increase hitch detection warning threshold from 100ms to 500ms.
- Group SBL + SBLCmd + SBLHelper into a "Programs/Switchboard" folder in the generated solution.

diff --git a/Engine/Source/Programs/SwitchboardListener/Private/CpuUtilizationMonitor.cpp b/Engine/Source/Programs/SwitchboardListener/Private/CpuUtilizationMonitor.cpp
deleted file mode 100644
diff --git a/Engine/Source/Programs/SwitchboardListener/Private/CpuUtilizationMonitor.h b/Engine/Source/Programs/SwitchboardListener/Private/CpuUtilizationMonitor.h
deleted file mode 100644
diff --git a/Engine/Source/Programs/SwitchboardListener/Private/SBLHelperClient.cpp b/Engine/Source/Programs/SwitchboardListener/Private/SBLHelperClient.cpp
deleted file mode 100644
diff --git a/Engine/Source/Programs/SwitchboardListener/Private/SBLHelperClient.h b/Engine/Source/Programs/SwitchboardListener/Private/SBLHelperClient.h
deleted file mode 100644
diff --git a/Engine/Source/Programs/SwitchboardListener/Private/SwitchboardAuth.cpp b/Engine/Source/Programs/SwitchboardListener/Private/SwitchboardAuth.cpp
deleted file mode 100644
diff --git a/Engine/Source/Programs/SwitchboardListener/Private/SwitchboardAuth.h b/Engine/Source/Programs/SwitchboardListener/Private/SwitchboardAuth.h
deleted file mode 100644
diff --git a/Engine/Source/Programs/SwitchboardListener/Private/SwitchboardListener.cpp b/Engine/Source/Programs/SwitchboardListener/Private/SwitchboardListener.cpp
deleted file mode 100644
diff --git a/Engine/Source/Programs/SwitchboardListener/Private/SwitchboardListener.h b/Engine/Source/Programs/SwitchboardListener/Private/SwitchboardListener.h
deleted file mode 100644
diff --git a/Engine/Source/Programs/SwitchboardListener/Private/SwitchboardListenerApp.h b/Engine/Source/Programs/SwitchboardListener/Private/SwitchboardListenerApp.h
deleted file mode 100644
diff --git a/Engine/Source/Programs/SwitchboardListener/Private/SwitchboardListenerMain.cpp b/Engine/Source/Programs/SwitchboardListener/Private/SwitchboardListenerMain.cpp
deleted file mode 100644
diff --git a/Engine/Source/Programs/SwitchboardListener/Private/SwitchboardListenerVersion.h b/Engine/Source/Programs/SwitchboardListener/Private/SwitchboardListenerVersion.h
deleted file mode 100644
diff --git a/Engine/Source/Programs/SwitchboardListener/Private/SwitchboardMessageFuture.h b/Engine/Source/Programs/SwitchboardListener/Private/SwitchboardMessageFuture.h
deleted file mode 100644
diff --git a/Engine/Source/Programs/SwitchboardListener/Private/SwitchboardPacket.h b/Engine/Source/Programs/SwitchboardListener/Private/SwitchboardPacket.h
deleted file mode 100644
diff --git a/Engine/Source/Programs/SwitchboardListener/Private/SwitchboardProtocol.cpp b/Engine/Source/Programs/SwitchboardListener/Private/SwitchboardProtocol.cpp
deleted file mode 100644
diff --git a/Engine/Source/Programs/SwitchboardListener/Private/SwitchboardProtocol.h b/Engine/Source/Programs/SwitchboardListener/Private/SwitchboardProtocol.h
deleted file mode 100644
diff --git a/Engine/Source/Programs/SwitchboardListener/Private/SwitchboardTasks.h b/Engine/Source/Programs/SwitchboardListener/Private/SwitchboardTasks.h
deleted file mode 100644
diff --git a/Engine/Source/Programs/SwitchboardListener/Private/SyncStatus.h b/Engine/Source/Programs/SwitchboardListener/Private/SyncStatus.h
deleted file mode 100644
diff --git a/Engine/Source/Programs/SwitchboardListener/Private/Windows/Resources/Resource.h b/Engine/Source/Programs/SwitchboardListener/Private/Windows/Resources/Resource.h
deleted file mode 100644
diff --git a/Engine/Source/Programs/SwitchboardListener/Private/Windows/Resources/Resource.rc b/Engine/Source/Programs/SwitchboardListener/Private/Windows/Resources/Resource.rc
deleted file mode 100644
diff --git a/Engine/Source/Programs/SwitchboardListener/Private/Windows/Resources/Resource.rc2 b/Engine/Source/Programs/SwitchboardListener/Private/Windows/Resources/Resource.rc2
deleted file mode 100644
diff --git a/Engine/Source/Programs/SwitchboardListener/Private/Windows/WindowsCpuUtilizationMonitor.cpp b/Engine/Source/Programs/SwitchboardListener/Private/Windows/WindowsCpuUtilizationMonitor.cpp
deleted file mode 100644
diff --git a/Engine/Source/Programs/SwitchboardListener/Private/Windows/WindowsCpuUtilizationMonitor.h b/Engine/Source/Programs/SwitchboardListener/Private/Windows/WindowsCpuUtilizationMonitor.h
deleted file mode 100644
diff --git a/Engine/Source/Programs/SwitchboardListener/SblCore/Internal/CpuUtilizationMonitor.h b/Engine/Source/Programs/SwitchboardListener/SblCore/Internal/CpuUtilizationMonitor.h
new file mode 100644
--- /dev/null
+++ b/Engine/Source/Programs/SwitchboardListener/SblCore/Internal/CpuUtilizationMonitor.h
@@ -0,0 +1,44 @@
+// Copyright Epic Games, Inc. All Rights Reserved.
+
+#pragma once
+
+#include "CoreTypes.h"
+#include "Containers/Array.h"
+#include "Logging/LogMacros.h"
+
+
+DECLARE_LOG_CATEGORY_EXTERN(LogCpuUtilizationMonitor, Log, All);
+
+
+// Default implementation does nothing (functionality optional and only implemented for some platforms).
+// Instead, use the FCpuUtilizationMonitor typedef to get a platform-specific implementation if available.
+class FGenericCpuUtilizationMonitor
+{
+public:
+	virtual ~FGenericCpuUtilizationMonitor()
+	{
+	}
+
+	virtual bool IsInitialized() const
+	{
+		return false;
+	}
+
+	/**
+	 * Samples the current CPU utilization, provided as per-core values ranging from 0-100.
+	 *
+	 * @param OutCoreUtilization The array that will receive the sampled values.
+	 * @return True if OutCoreUtilization is valid (new or recently cached values were available).
+	 */
+	virtual bool GetPerCoreUtilization(TArray<int8>& OutCoreUtilization)
+	{
+		return false;
+	}
+};
+
+
+#if PLATFORM_WINDOWS
+#include COMPILED_PLATFORM_HEADER(CpuUtilizationMonitor.h)
+#else
+typedef FGenericCpuUtilizationMonitor FCpuUtilizationMonitor;
+#endif
diff --git a/Engine/Source/Programs/SwitchboardListener/SblCore/Internal/SwitchboardListener.h b/Engine/Source/Programs/SwitchboardListener/SblCore/Internal/SwitchboardListener.h
new file mode 100644
--- /dev/null
+++ b/Engine/Source/Programs/SwitchboardListener/SblCore/Internal/SwitchboardListener.h
@@ -0,0 +1,294 @@
+// Copyright Epic Games, Inc. All Rights Reserved.
+
+#pragma once
+
+#include "SwitchboardListenerApp.h"
+#include "Async/Mutex.h"
+#include "Async/RecursiveMutex.h"
+#include "Containers/Queue.h"
+#include "HAL/PlatformProcess.h"
+#include "Interfaces/IPv4/IPv4Endpoint.h"
+
+#include "CpuUtilizationMonitor.h"
+
+
+#if !PLATFORM_WINDOWS
+#	include <msquic.h>
+#else
+#	include "Windows/AllowWindowsPlatformTypes.h"
+#	include <msquic.h>
+#	include "Windows/HideWindowsPlatformTypes.h"
+#endif
+
+
+struct FRunningProcess;
+struct FSwitchboardMessageFuture;
+struct FSwitchboardTask;
+struct FSwitchboardAuthenticateTask;
+struct FSwitchboardDisconnectTask;
+struct FSwitchboardSendFileToClientTask;
+struct FSwitchboardStartTask;
+struct FSwitchboardKillTask;
+struct FSwitchboardReceiveFileFromClientTask;
+struct FSwitchboardGetSyncStatusTask;
+struct FSwitchboardRefreshMosaicsTask;
+struct FSwitchboardRedeployListenerTask;
+struct FSwitchboardFixExeFlagsTask;
+struct FSwitchboardMinimizeWindowsTask;
+struct FSwitchboardSetInactiveTimeoutTask;
+struct FSwitchboardFreeListenerBinaryTask;
+
+class FSBLHelperClient;
+
+
+struct FSwitchboardCommandLineOptions
+{
+	bool OutputVersion = false;
+	bool MinimizeOnLaunch = true;
+
+	TOptional<FIPv4Address> Address;
+	TOptional<uint16> Port;
+
+	enum class ESecureMode : uint8
+	{
+		Unspecified = 0,
+		CertificateFromFile,
+#if PLATFORM_WINDOWS
+		CertificateByHash,
+#endif
+	};
+
+	ESecureMode SecureMode;
+	TOptional<FString> CertificateFile;
+	TOptional<FString> PrivateKeyFile;
+#if PLATFORM_WINDOWS
+	TOptional<FString> CertificateHash;
+#endif
+
+	TOptional<uint32> RedeployFromPid;
+
+	static FSwitchboardCommandLineOptions FromString(const TCHAR* CommandLine);
+	FString ToString(bool bIncludeRedeploy = false) const;
+};
+
+
+struct FRedeployStatus
+{
+	enum class EState : uint8
+	{
+		// No redeploy initiated.
+		None,
+
+		// Redeploy initiated by client, decline subsequent redeploy requests.
+		RequestReceived,
+
+		// New listener executable hash verified and written to randomly-named temporary file.
+		// This also means we've stopped listening for new connections in this (old) listener.
+		NewListenerWrittenTemp,
+
+		// New listener successfully launched and running.
+		NewListenerStarted,
+
+		// Current process executable moved to randomly-named temporary file.
+		ThisListenerRenamed,
+
+		// New listener executable moved to original location of this executable.
+		NewListenerRenamed,
+
+		// Child process has signaled via IPC that it's initialized; shutting down.
+		Complete,
+	};
+
+	bool InProgress() const
+	{
+		return State != EState::None && State != EState::Complete;
+	}
+
+	EState State = EState::None;
+
+	// Valid when State >= RequestReceived
+	FIPv4Endpoint RequestingClient;
+
+	// Valid when State >= NewListenerWrittenTemp
+	FString OriginalThisExePath;
+	FString TempNewListenerPath;
+
+	// Valid when State >= NewListenerStarted
+	FProcHandle ListenerProc;
+
+	// Valid when State >= ThisListenerRenamed
+	FString TempThisListenerPath;
+};
+
+
+class FSwitchboardListener
+{
+	struct FConnection;
+	using FConnectionRef = TSharedRef<FConnection>;
+	using FByteArrayRef = TSharedRef<TArray<uint8>>;
+
+	static const FIPv4Endpoint InvalidEndpoint;
+	static constexpr const TCHAR* PasswordCredentialName = TEXT("PresharedAuthToken");
+
+public:
+	explicit FSwitchboardListener(const FSwitchboardCommandLineOptions& InOptions);
+	~FSwitchboardListener();
+
+	bool Init(bool bStartListeningImmediately = true);
+	void Shutdown();
+
+	void Tick();
+
+	bool StartListening();
+	bool StopListening();
+	bool IsListening() const { return QuicApi != nullptr; }
+
+	const FString& GetAuthPassword() const { return ExpectedAuthPassword; }
+	bool SetAuthPassword(const FString& NewPassword);
+
+	TSet<FIPv4Endpoint> GetConnectedClientEndpoints() const;
+	TSet<FIPv4Address> GetConnectedClientAddresses() const;
+
+	// Used for redeploy synchronization.
+	static FString GetIpcSemaphoreName(uint32 ParentPid);
+
+	FSimpleMulticastDelegate& OnInit() { return OnInitDelegate; }
+	FSimpleMulticastDelegate& OnShutdown() { return OnShutdownDelegate; }
+	FSimpleMulticastDelegate& OnTick() { return OnTickDelegate; }
+
+private:
+
+	bool ParseIncomingMessage(const FString& InMessage, const FIPv4Endpoint& InEndpoint, const FConnectionRef& Connection);
+
+	bool RunScheduledTask(const FSwitchboardTask& InTask);
+	bool Task_Authenticate(const FSwitchboardAuthenticateTask& InAuthTask);
+	bool Task_StartProcess(const FSwitchboardStartTask& InRunTask);
+	bool Task_KillProcess(const FSwitchboardKillTask& KillTask);
+	bool Task_ReceiveFileFromClient(const FSwitchboardReceiveFileFromClientTask& InReceiveFileFromClientTask);
+	bool Task_RedeployListener(const FSwitchboardRedeployListenerTask& InRedeployListenerTask);
+	bool Task_SendFileToClient(const FSwitchboardSendFileToClientTask& InSendFileToClientTask);
+	bool Task_GetSyncStatus(const FSwitchboardGetSyncStatusTask& InGetSyncStatusTask);
+	bool Task_RefreshMosaics(const FSwitchboardRefreshMosaicsTask& InRefreshMosaicsTask);
+	bool Task_FixExeFlags(const FSwitchboardFixExeFlagsTask& InFixExeFlagsTask);
+	bool Task_MinimizeWindows(const FSwitchboardMinimizeWindowsTask& InRefreshMosaicsTask);
+	bool Task_SetInactiveTimeout(const FSwitchboardSetInactiveTimeoutTask& InTimeoutTask);
+	bool Task_FreeListenerBinary(const FSwitchboardFreeListenerBinaryTask& InRenameProcessTask);
+
+	bool KillProcessNow(FRunningProcess* InProcess, float SoftKillTimeout = 0.0f);
+	FRunningProcess* FindOrStartFlipModeMonitorForUUID(const FGuid& UUID);
+
+	void CleanUpDisconnectedSockets();
+	void HandleStdout(const TSharedPtr<FRunningProcess>& Process);
+	void HandleRunningProcesses(TArray<TSharedPtr<FRunningProcess>>& Processes, bool bNotifyThatProgramEnded);
+
+	bool SendMessage(const FString& InMessage, const FIPv4Endpoint& InEndpoint);
+	void SendMessageFutures();
+
+	bool EquivalentTaskFutureExists(uint32 TaskEquivalenceHash) const;
+
+	void RollbackRedeploy();
+
+	_Function_class_(QUIC_LISTENER_CALLBACK)
+	static QUIC_STATUS QUIC_API QuicListenerThunk(HQUIC Listener, void* Context, QUIC_LISTENER_EVENT* Event);
+	static QUIC_STATUS QUIC_API QuicConnectionThunk(HQUIC Connection, void* Context, QUIC_CONNECTION_EVENT* Event);
+	static QUIC_STATUS QUIC_API QuicStreamThunk(HQUIC Stream, void* Context, QUIC_STREAM_EVENT* Event);
+	QUIC_STATUS QuicListenerCallback(HQUIC Listener, QUIC_LISTENER_EVENT* Event);
+	QUIC_STATUS QuicConnectionCallback(HQUIC Connection, QUIC_CONNECTION_EVENT* Event);
+	QUIC_STATUS QuicStreamCallback(HQUIC Stream, QUIC_STREAM_EVENT* Event);
+
+private: // Private nested types.
+	/** Combines a QUIC_BUFFER span (ptr + len) and its backing array. */
+	struct FQuicBuffer
+	{
+		QUIC_BUFFER QuicBuffer;
+		FByteArrayRef Storage;
+
+		FQuicBuffer()
+			: QuicBuffer{ .Length = 0, .Buffer = nullptr }
+			, Storage(MakeShared<TArray<uint8>>())
+		{
+		}
+
+		explicit FQuicBuffer(FByteArrayRef&& InStorage)
+			: Storage(MoveTemp(InStorage))
+		{
+			QuicBuffer = {
+				.Length = static_cast<uint32>(Storage->Num()),
+				.Buffer = Storage->GetData(),
+			};
+		}
+	};
+
+	/** Tracks all the state for an established connection. */
+	struct FConnection
+	{
+		FIPv4Endpoint Endpoint;
+		HQUIC QuicConn = nullptr;
+		// We currently assume a single bidirectional stream initiated by the peer.
+		HQUIC QuicStream = nullptr;
+
+		UE::FMutex SendLock;
+		TQueue<TSharedPtr<FQuicBuffer>> SendBuffers;
+
+		UE::FMutex ReceiveLock;
+		bool bMessageComplete = false;
+		FByteArrayRef ReceiveBuffer = MakeShared<TArray<uint8>>();
+
+		bool bAuthenticated = false;
+
+		// TODO?: MsQuic handles idle timeout, but configuration is app-wide.
+		// If we want to support the existing per-client override, need to
+		// figure out how best to pass one override when calling
+		// QuicApi->ConnectionSetConfiguration() (clone config obj N times?).
+		float InactiveTimeout;
+		double LastActivityTime;
+	};
+
+private:
+	FSwitchboardCommandLineOptions Options;
+	TUniquePtr<FIPv4Endpoint> ListenerEndpoint;
+
+	TOptional<FString> PrivateKeyPassword;
+	FString ExpectedAuthPassword;
+
+	/** MsQuic top level function table for all other API calls. */
+	const QUIC_API_TABLE* QuicApi = nullptr;
+
+	/** MsQuic registration manages the execution context for all child objects. */
+	HQUIC QuicRegistration = nullptr;
+
+	/** MsQuic configuration manages security-related and other common QUIC settings. */
+	HQUIC QuicConfiguration = nullptr;
+
+	/** MsQuic incoming connection handler. */
+	HQUIC QuicListener = nullptr;
+
+	UE::FRecursiveMutex ConnectionsLock;
+	TMap<FIPv4Endpoint, FConnectionRef> ConnectionsByEndpoint;
+	TMap<HQUIC, FConnectionRef> ConnectionsByQuicConn;
+	TMap<HQUIC, FConnectionRef> ConnectionsByQuicStream;
+
+	static constexpr int8 MaxAuthFailures = 5;
+	TMap<FIPv4Address, int8> AuthFailuresByAddress;
+
+	TQueue<TUniquePtr<FSwitchboardTask>, EQueueMode::Spsc> ScheduledTasks;
+	TQueue<TUniquePtr<FSwitchboardTask>, EQueueMode::Spsc> DisconnectTasks;
+	TArray<TSharedPtr<FRunningProcess>> RunningProcesses;
+	TArray<TSharedPtr<FRunningProcess>> FlipModeMonitors;
+	TArray<FSwitchboardMessageFuture> MessagesFutures;
+	TSharedPtr<FCpuUtilizationMonitor> CpuMonitor;
+
+	/** Client interface to the Switchboard Listener Helper external process */
+	TSharedPtr<FSBLHelperClient> SBLHelper;
+
+	bool bIsNvAPIInitialized;
+
+	TSharedPtr<FRWLock> CachedMosaicToposLock;
+	TSharedPtr<TArray<struct FMosaicTopo>> CachedMosaicTopos;
+
+	FRedeployStatus RedeployStatus;
+
+	FSimpleMulticastDelegate OnInitDelegate;
+	FSimpleMulticastDelegate OnShutdownDelegate;
+	FSimpleMulticastDelegate OnTickDelegate;
+};
diff --git a/Engine/Source/Programs/SwitchboardListener/SblCore/Internal/SwitchboardListenerApp.h b/Engine/Source/Programs/SwitchboardListener/SblCore/Internal/SwitchboardListenerApp.h
new file mode 100644
--- /dev/null
+++ b/Engine/Source/Programs/SwitchboardListener/SblCore/Internal/SwitchboardListenerApp.h
@@ -0,0 +1,43 @@
+// Copyright Epic Games, Inc. All Rights Reserved.
+
+#pragma once
+
+#include "Logging/LogMacros.h"
+
+
+DECLARE_LOG_CATEGORY_EXTERN(LogSwitchboard, Verbose, All);
+
+
+#if !defined(SBL_SLATE)
+#	define SBL_SLATE 0
+#endif
+
+
+#if !defined(SWITCHBOARD_CPUPROFILERTRACE_ENABLED)
+#	if CPUPROFILERTRACE_ENABLED
+#		define SWITCHBOARD_CPUPROFILERTRACE_ENABLED 1
+#	else
+#		define SWITCHBOARD_CPUPROFILERTRACE_ENABLED 0
+#	endif
+#endif
+
+
+#if SWITCHBOARD_CPUPROFILERTRACE_ENABLED
+#	include "ProfilingDebugging/CpuProfilerTrace.h"
+
+#	define SWITCHBOARD_TRACE_CPUPROFILER_EVENT_SCOPE(Name) \
+		TRACE_CPUPROFILER_EVENT_SCOPE(Name)
+
+#	define SWITCHBOARD_TRACE_CPUPROFILER_EVENT_SCOPE_STR(Name) \
+		TRACE_CPUPROFILER_EVENT_SCOPE_STR(Name)
+
+#	define SWITCHBOARD_TRACE_CPUPROFILER_EVENT_SCOPE_TEXT(Name) \
+		TRACE_CPUPROFILER_EVENT_SCOPE_TEXT(Name)
+
+#else
+
+#	define SWITCHBOARD_TRACE_CPUPROFILER_EVENT_SCOPE(Name)
+#	define SWITCHBOARD_TRACE_CPUPROFILER_EVENT_SCOPE_STR(Name)
+#	define SWITCHBOARD_TRACE_CPUPROFILER_EVENT_SCOPE_TEXT(Name)
+
+#endif
diff --git a/Engine/Source/Programs/SwitchboardListener/SblCore/Internal/SwitchboardListenerMainCommon.h b/Engine/Source/Programs/SwitchboardListener/SblCore/Internal/SwitchboardListenerMainCommon.h
new file mode 100644
--- /dev/null
+++ b/Engine/Source/Programs/SwitchboardListener/SblCore/Internal/SwitchboardListenerMainCommon.h
@@ -0,0 +1,13 @@
+// Copyright Epic Games, Inc. All Rights Reserved.
+
+#pragma once
+
+#include "HAL/Platform.h"
+
+
+struct FSwitchboardCommandLineOptions;
+
+bool SwitchboardListenerMainInit(const FSwitchboardCommandLineOptions& InOptions);
+bool SwitchboardListenerMainShutdown();
+
+int32 SwitchboardListenerMainWrapper();
diff --git a/Engine/Source/Programs/SwitchboardListener/SblCore/Internal/SwitchboardListenerVersion.h b/Engine/Source/Programs/SwitchboardListener/SblCore/Internal/SwitchboardListenerVersion.h
new file mode 100644
--- /dev/null
+++ b/Engine/Source/Programs/SwitchboardListener/SblCore/Internal/SwitchboardListenerVersion.h
@@ -0,0 +1,7 @@
+// Copyright Epic Games, Inc. All Rights Reserved.
+
+#pragma once
+
+#define SBLISTENER_VERSION_MAJOR 3
+#define SBLISTENER_VERSION_MINOR 0
+#define SBLISTENER_VERSION_PATCH 2
diff --git a/Engine/Source/Programs/SwitchboardListener/SblCore/Internal/Windows/WindowsCpuUtilizationMonitor.h b/Engine/Source/Programs/SwitchboardListener/SblCore/Internal/Windows/WindowsCpuUtilizationMonitor.h
new file mode 100644
--- /dev/null
+++ b/Engine/Source/Programs/SwitchboardListener/SblCore/Internal/Windows/WindowsCpuUtilizationMonitor.h
@@ -0,0 +1,42 @@
+// Copyright Epic Games, Inc. All Rights Reserved.
+
+#pragma once
+
+#ifndef _WIN32_WINDOWS
+#define _WIN32_WINDOWS 0x0410
+#endif
+
+#include "CpuUtilizationMonitor.h"
+
+#include "Windows/AllowWindowsPlatformTypes.h"
+#include <Windows.h>
+#include <Pdh.h>
+#include <PdhMsg.h>
+#include "Windows/HideWindowsPlatformTypes.h"
+
+
+class FWindowsCpuUtilizationMonitor : public FGenericCpuUtilizationMonitor
+{
+public:
+	FWindowsCpuUtilizationMonitor();
+	~FWindowsCpuUtilizationMonitor();
+
+	bool IsInitialized() const override
+	{
+		return bIsInitialized;
+	}
+
+	bool GetPerCoreUtilization(TArray<int8>& OutCoreUtilization) override;
+
+private:
+	bool QueryUpdatedUtilization();
+
+	bool bIsInitialized = false;
+
+	HQUERY QueryHandle = INVALID_HANDLE_VALUE;
+	TArray<PDH_HCOUNTER> CounterHandles;
+	double LastQueryTime = DBL_MIN;
+	TArray<int8> CachedCoreUtilization;
+};
+
+typedef FWindowsCpuUtilizationMonitor FCpuUtilizationMonitor;
diff --git a/Engine/Source/Programs/SwitchboardListener/SblCore/Private/CpuUtilizationMonitor.cpp b/Engine/Source/Programs/SwitchboardListener/SblCore/Private/CpuUtilizationMonitor.cpp
new file mode 100644
--- /dev/null
+++ b/Engine/Source/Programs/SwitchboardListener/SblCore/Private/CpuUtilizationMonitor.cpp
@@ -0,0 +1,6 @@
+// Copyright Epic Games, Inc. All Rights Reserved.
+
+#include "CpuUtilizationMonitor.h"
+
+
+DEFINE_LOG_CATEGORY(LogCpuUtilizationMonitor);
diff --git a/Engine/Source/Programs/SwitchboardListener/SblCore/Private/SBLHelperClient.cpp b/Engine/Source/Programs/SwitchboardListener/SblCore/Private/SBLHelperClient.cpp
new file mode 100644
--- /dev/null
+++ b/Engine/Source/Programs/SwitchboardListener/SblCore/Private/SBLHelperClient.cpp
@@ -0,0 +1,339 @@
+// Copyright Epic Games, Inc. All Rights Reserved.
+
+#include "SBLHelperClient.h"
+
+#include "Common/TcpSocketBuilder.h"
+#include "Dom/JsonValue.h"
+#include "JsonObjectConverter.h"
+#include "Policies/CondensedJsonPrintPolicy.h"
+#include "Serialization/JsonReader.h"
+#include "Serialization/JsonSerializer.h"
+#include "Serialization/JsonWriter.h"
+
+
+
+DEFINE_LOG_CATEGORY(LogSBLHelperClient);
+
+const int32 FSBLHelperClient::CompatibleVersion[3] = { 1, 0, 0 };
+
+FSBLHelperClient::~FSBLHelperClient()
+{
+	Disconnect();
+}
+
+bool FSBLHelperClient::IsConnected() 
+{
+	int32 BytesSent;
+	return Socket && Socket->Send(nullptr, 0, BytesSent);
+}
+
+bool FSBLHelperClient::Connect(const FConnectionParams& InConnectionParams)
+{
+	check(IsInGameThread());
+	
+	// Cache connection parameters
+	ConnectionParams = InConnectionParams;
+
+	// Make sure we have released everything related to any previous connection attempts
+	Disconnect();
+	check(!Socket);
+
+	// Create socket for the new connection
+
+	Socket = FTcpSocketBuilder(TEXT("SBLHelperSocket"));
+	check(Socket);
+
+	// Attempt to connect
+	{
+		const bool bConnected = Socket->Connect(*ConnectionParams.Endpoint.ToInternetAddr());
+
+		if (!bConnected)
+		{
+			Disconnect();
+			return false;
+		}
+	}
+
+	// Confirm that this is actually a SBL server by parsing its first packet(s), if any.
+	constexpr float TimeoutSeconds = 0.5;
+	ParseIncomingMessages(TimeoutSeconds);
+
+	// If the hello packet arrived, ServerVersion must have been populated, which we verify below.
+	{
+		TArray<FString> VersionStringArray;
+		ServerVersion.ParseIntoArray(VersionStringArray, TEXT("."));
+
+		if (VersionStringArray.Num() != 3)
+		{
+			UE_LOG(LogSBLHelperClient, Error, TEXT("Invalid SBLHelperService version '%s' at endpoint %s"), 
+				*ServerVersion,
+				*ConnectionParams.Endpoint.ToString()
+			);
+
+			Disconnect();
+			return false;
+		}
+
+		const int32 VersionArray[3] = {
+			FCString::Atoi(*VersionStringArray[0]),
+			FCString::Atoi(*VersionStringArray[1]),
+			FCString::Atoi(*VersionStringArray[2]),
+		};
+
+		// We'll only look at major and minor.
+		if ((VersionArray[0] != CompatibleVersion[0]) 
+			|| (VersionArray[1] < CompatibleVersion[1]))
+		{
+			UE_LOG(LogSBLHelperClient, Error, TEXT("Incompatible SBLHelperService version '%s' at endpoint %s. Minimum required is %d.%d.%d"), 
+				*ServerVersion,
+				*ConnectionParams.Endpoint.ToString(),
+				CompatibleVersion[0],
+				CompatibleVersion[1],
+				CompatibleVersion[2]
+			);
+
+			Disconnect();
+			return false;
+		}
+	}
+
+	// Ok, we're good to go.
+	UE_LOG(LogSBLHelperClient, Display, TEXT("Connected to SBLHelperService version '%s' at endpoint %s"),
+		*ServerVersion,
+		*ConnectionParams.Endpoint.ToString()
+	);
+
+	return true;
+}
+
+void FSBLHelperClient::HandleCmdHello(TSharedPtr<FJsonObject>& Json)
+{
+	check(Json.IsValid());
+
+	// We use "sblhver" instead of just "ver" as a magic value in disguise.
+	TSharedPtr<FJsonValue> VersionField = Json->TryGetField(TEXT("sblhver"));
+
+	if (!VersionField.IsValid())
+	{
+		return;
+	}
+
+	ServerVersion = VersionField->AsString();
+}
+
+void FSBLHelperClient::ParseMessage(const FString& Message)
+{
+	// All messages are in Json format
+
+	TSharedRef<TJsonReader<TCHAR>> Reader = FJsonStringReader::Create(Message);
+	TSharedPtr<FJsonObject> JsonData;
+
+	if (!FJsonSerializer::Deserialize(Reader, JsonData))
+	{
+		return;
+	}
+
+	// All messages must have a cmd field.
+	TSharedPtr<FJsonValue> CmdField = JsonData->TryGetField(TEXT("cmd"));
+	if (!CmdField.IsValid())
+	{
+		return;
+	}
+
+	const FString Cmd = CmdField->AsString();
+
+	// Each command has its own handler
+	if (Cmd.Equals("hello"))
+	{
+		HandleCmdHello(JsonData);
+	}
+	else
+	{
+		UE_LOG(LogSBLHelperClient, Error, TEXT("Unhandled command '%s'"), *Cmd);
+	}
+}
+
+void FSBLHelperClient::ParseIncomingMessages(const float TimeoutSeconds)
+{
+	for (const FString& Message : ReceiveMessages(TimeoutSeconds))
+	{
+		ParseMessage(Message);
+	}
+}
+
+void FSBLHelperClient::Tick()
+{
+	if (!Socket)
+	{
+		return;
+	}
+
+	// We parse and handle all incoming messages.
+	ParseIncomingMessages();
+}
+
+TArray<FString> FSBLHelperClient::ReceiveMessages(const float TimeoutSeconds)
+{
+	TArray<FString> Messages;
+
+	if (!Socket)
+	{
+		return Messages;
+	}
+
+	const double TimeoutTime = FPlatformTime::Seconds() + TimeoutSeconds;
+
+	do // Execute at least once
+	{
+		uint32 PendingDataSize = 0;
+
+		while (Socket->HasPendingData(PendingDataSize))
+		{
+			// Read 
+			TArray<uint8> Buffer;
+			Buffer.AddUninitialized(PendingDataSize);
+			int32 BytesRead = 0;
+
+			if (!Socket->Recv(Buffer.GetData(), PendingDataSize, BytesRead, ESocketReceiveFlags::None))
+			{
+				UE_LOG(LogSBLHelperClient, Error, TEXT("Error while receiving data via endpoint %s"), 
+					*ConnectionParams.Endpoint.ToString()
+				);
+
+				continue;
+			}
+
+			for (int32 ByteIdx = 0; ByteIdx < BytesRead; ++ByteIdx)
+			{
+				ReceiveBuffer.Add(Buffer[ByteIdx]);
+
+				if (Buffer[ByteIdx] == '\x00')
+				{
+					Messages.Emplace(UTF8_TO_TCHAR(ReceiveBuffer.GetData()));
+
+					// Clear the reception buffer so that it only contains the next partial message
+					ReceiveBuffer.Empty();
+				}
+			}
+		}
+
+		// If we already got at least one message, we can return.
+		if (Messages.Num())
+		{
+			return Messages;
+		}
+
+		// We don't wait more that the given timeout
+		const double CurrentTime = FPlatformTime::Seconds();
+		const double RemainingTimeBeforeTimeout = TimeoutTime - CurrentTime;
+
+		if (RemainingTimeBeforeTimeout <= 0)
+		{
+			break;
+		}
+
+		// Sleep for a bit but not to exceed our timeout time.
+		const double SleepSeconds = FMath::Min(0.025, RemainingTimeBeforeTimeout);
+		FPlatformProcess::SleepNoStats(SleepSeconds);
+
+	} while (true);
+
+	return Messages;
+}
+
+uint32 FSBLHelperClient::GetNextPacketId()
+{
+	NextPacketId++;
+
+	// Avoid 0, which is reserved for no packet id.
+
+	if (NextPacketId == 0)
+	{
+		NextPacketId = 1;
+	}
+
+	return NextPacketId;
+}
+
+FString FSBLHelperClient::CreateMessage(const FString& Cmd, const TMap<FString, FString>& AdditionalFields)
+{
+	FString Message;
+
+	TSharedRef<TJsonWriter<TCHAR, TCondensedJsonPrintPolicy<TCHAR>>> JsonWriter 
+		= TJsonWriterFactory<TCHAR, TCondensedJsonPrintPolicy<TCHAR>>::Create(&Message);
+
+	JsonWriter->WriteObjectStart();
+
+	// message id
+	JsonWriter->WriteValue(TEXT("id"), FString::Printf(TEXT("%u"), GetNextPacketId()));
+
+	// cmd
+	JsonWriter->WriteValue(TEXT("cmd"), Cmd);
+
+	for (const auto& Value : AdditionalFields)
+	{
+		JsonWriter->WriteValue(Value.Key, Value.Value);
+	}
+
+	JsonWriter->WriteObjectEnd();
+	JsonWriter->Close();
+
+	return Message;
+}
+
+void FSBLHelperClient::Disconnect()
+{
+	// Send by message
+	SendMessage(CreateMessage(TEXT("bye"), TMap<FString, FString>()));
+
+	// Clear buffer so that we don't try to use leftover data on the next connection
+	ReceiveBuffer.Empty();
+
+	// Clear the version as well, for the same reason.
+	ServerVersion.Empty();
+
+	// Close and destroy socket
+	if (Socket)
+	{
+		Socket->Close();
+		ISocketSubsystem::Get()->DestroySocket(Socket);
+		Socket = nullptr;
+	}
+}
+
+bool FSBLHelperClient::LockGpuClock(const uint32 Pid)
+{
+	// Send the lock command and the related process pid.
+
+	TMap<FString, FString> AdditionalFields;
+	AdditionalFields.Emplace(TEXT("pid"), FString::Printf(TEXT("%u"), Pid));
+
+	return SendMessage(CreateMessage(TEXT("lock"), AdditionalFields));
+}
+
+bool FSBLHelperClient::SendMessage(const FString& Message)
+{
+	if (!Socket)
+	{
+		return false;
+	}
+
+	UE_LOG(LogSBLHelperClient, Verbose, TEXT("Sending message: %s"), *Message);
+
+	// Send the message over the socket. 
+	// Note the "+ 1" to include the null terminator, which is the message delimiter in the SBLHelper protocol.
+
+	int32 BytesSent = 0;
+	return Socket->Send((uint8*)TCHAR_TO_UTF8(*Message), Message.Len() + 1, BytesSent);
+}
+
+FString FSBLHelperClient::GetStatus()
+{
+	if (!IsConnected())
+	{
+		return TEXT("Disconnected");
+	}
+
+	return TEXT("Connected");
+}
+
diff --git a/Engine/Source/Programs/SwitchboardListener/SblCore/Private/SBLHelperClient.h b/Engine/Source/Programs/SwitchboardListener/SblCore/Private/SBLHelperClient.h
new file mode 100644
--- /dev/null
+++ b/Engine/Source/Programs/SwitchboardListener/SblCore/Private/SBLHelperClient.h
@@ -0,0 +1,124 @@
+// Copyright Epic Games, Inc. All Rights Reserved.
+
+#pragma once
+
+#include "Containers/Array.h"
+#include "Containers/UnrealString.h"
+#include "Interfaces/IPv4/IPv4Endpoint.h"
+#include "Logging/LogMacros.h"
+#include "Templates/SharedPointer.h"
+
+DECLARE_LOG_CATEGORY_EXTERN(LogSBLHelperClient, Log, All);
+
+class FJsonObject;
+class FSocket;
+
+/** Class in charge of communicating with the SwitchboardListenerHelper server */
+class FSBLHelperClient
+{
+public:
+	
+	/** Structure holding the desired connection parameters, such as IP and Port */
+	struct FConnectionParams
+	{
+		/** Server end point to connect to */
+		FIPv4Endpoint Endpoint;
+	};
+
+	/** Destructor */
+	~FSBLHelperClient();
+
+	/** Returns true if currently connected to the SBLHelper server */
+	bool IsConnected();
+
+	/** 
+	 * Connects to the SBLHelper server specified in the given connection parameters. 
+	 * If already connected, it will disconnect before attempting the new connection.
+	 * 
+	 * @return True if it was able to establish the connection to a valid server.
+	 */
+	bool Connect(const FConnectionParams& InConnectionParams);
+
+	/** Disconnects from the SBLHelper server. Can be called even if already disconnected. */
+	void Disconnect();
+
+	/** 
+	 * Requests the server to lock the gpu and memory clocks to the maximum allowed. 
+	 * 
+	 * @param Pid Proces ID to which the lock gpu clock request is tied to.
+	 * 
+	 * @return True if the message was sent over the socket. 
+	 */
+	bool LockGpuClock(const uint32 Pid);
+
+	/** Returns a string describing the status of the server */
+	FString GetStatus();
+
+	/** Call this periodically for tasks such as parsing the incoming messages from the server */
+	void Tick();
+
+private:
+
+	/** 
+	 * Receives messages from the socket into an array. 
+	 * 
+	 * @param TimeoutSeconds Maximum time to wait for at least one message to be completely received.
+	 * 
+	 * @return Array of received messages.
+	 */
+	TArray<FString> ReceiveMessages(const float TimeoutSeconds = 0);
+
+	/** 
+	 * Parses the given message and calls the respective handlers 
+	 * 
+	 * @param InMessage The message to parse.
+	 */
+	void ParseMessage(const FString& InMessage);
+
+	/** 
+	 * Checks the socket for new messages and calls ParseMessage on each.
+	 * 
+	 * @param TimeoutSeconds The maximum time to wait for at least one message to be completely received.
+	 */
+	void ParseIncomingMessages(const float TimeoutSeconds = 0);
+
+	/** Handles the 'Hello' command from the server */
+	void HandleCmdHello(TSharedPtr<FJsonObject>& Json);
+
+	/** 
+	 * Creates a new message for the given command and additional fields 
+	 * 
+	 * @param Cmd The command to create this message for.
+	 * @param AdditionalFields Any additional fields or parameters to include in the message.
+	 * 
+	 * @return The created message, ready to be sent over the socket.
+	 */
+	FString CreateMessage(const FString& Cmd, const TMap<FString, FString>& AdditionalFields);
+
+	/** Sends the given message over the socket, if connected. */
+	bool SendMessage(const FString& Message);
+
+	/** Gets the next packet id to use when creating a new message. Monotonically increments. Skips zero. */
+	uint32 GetNextPacketId();
+
+private:
+
+	/** The minimum compatible SBLHelper server version.  */
+	static const int32 CompatibleVersion[3];
+
+	/** Socket for this client to communicate with the server */
+	FSocket* Socket = nullptr;
+
+	/** Next packet id when creating a new message to be sent */
+	uint32 NextPacketId = 1;
+
+	/** Buffer to hold partial messages received over the socket from the server */
+	TArray<uint8> ReceiveBuffer;
+
+	/** Cached connection parameters. Mostly used for logging purposes. */
+	FConnectionParams ConnectionParams;
+
+	/** The current server version that this client is connected to. */
+	FString ServerVersion;
+};
+
diff --git a/Engine/Source/Programs/SwitchboardListener/SblCore/Private/SwitchboardAuth.cpp b/Engine/Source/Programs/SwitchboardListener/SblCore/Private/SwitchboardAuth.cpp
new file mode 100644
--- /dev/null
+++ b/Engine/Source/Programs/SwitchboardListener/SblCore/Private/SwitchboardAuth.cpp
@@ -0,0 +1,437 @@
+// Copyright Epic Games, Inc. All Rights Reserved.
+
+#include "SwitchboardAuth.h"
+#include "SwitchboardListenerApp.h"
+
+#include "Containers/StringConv.h"
+#include "HAL/FileManagerGeneric.h"
+#include "HAL/PlatformProcess.h"
+#include "Logging/StructuredLog.h"
+#include "Misc/FileHelper.h"
+#include "Misc/Paths.h"
+#include "Misc/ScopeExit.h"
+
+#include <iostream>
+#include <string>
+
+#if PLATFORM_WINDOWS
+#include "Windows/AllowWindowsPlatformTypes.h"
+#include <Windows.h>
+#include <ShlObj.h>
+#include <wincred.h>
+#include "Windows/HideWindowsPlatformTypes.h"
+#else
+#include <sys/file.h>
+#include <termios.h>
+#endif
+
+UE_PUSH_MACRO("UI")
+// Workaround for "ossl_typ.h(144): error C2365: 'UI': redefinition; previous definition was 'namespace'"
+//                "ObjectMacros.h(872): note: see declaration of 'UI'"
+#define UI UI_ST
+#include <openssl/pem.h>
+#include <openssl/rand.h>
+#include <openssl/x509.h>
+UE_POP_MACRO("UI")
+
+
+namespace UE::SwitchboardListener
+{
+	FString ReadCredentialFromStdin(const FString& CredentialDesc, bool bAllowEmpty /* = false */)
+	{
+		FString Credential;
+
+		// Disable echo during token input
+#if PLATFORM_WINDOWS
+		HANDLE Stdin = GetStdHandle(STD_INPUT_HANDLE);
+		DWORD PrevConsoleMode = 0;
+		GetConsoleMode(Stdin, &PrevConsoleMode);
+		SetConsoleMode(Stdin, PrevConsoleMode & ~ENABLE_ECHO_INPUT);
+#else
+		struct termios Term;
+		tcgetattr(fileno(stdin), &Term);
+		Term.c_lflag &= ~ECHO;
+		tcsetattr(fileno(stdin), TCSADRAIN, &Term);
+#endif
+
+		ON_SCOPE_EXIT
+		{
+			// Restore console echo
+#if PLATFORM_WINDOWS
+			SetConsoleMode(Stdin, PrevConsoleMode);
+#else
+			Term.c_lflag |= ECHO;
+			tcsetattr(fileno(stdin), TCSADRAIN, &Term);
+#endif
+		};
+
+		while (Credential.IsEmpty())
+		{
+			// FIXME?: Doesn't seem like we have a suitable abstraction for this?
+			std::string CredFirst, CredSecond;
+
+			printf("\nEnter %s: ", TCHAR_TO_UTF8(*CredentialDesc));
+			std::getline(std::cin, CredFirst);
+
+			if (CredFirst.empty() && !bAllowEmpty)
+			{
+				printf("\nEmpty %s is not allowed!", TCHAR_TO_UTF8(*CredentialDesc));
+				continue;
+			}
+
+			printf("\nConfirm %s: ", TCHAR_TO_UTF8(*CredentialDesc));
+			std::getline(std::cin, CredSecond);
+
+			if (CredFirst == CredSecond)
+			{
+				printf("\n\n");
+				Credential = UTF8_TO_TCHAR(CredFirst.c_str());
+			}
+			else
+			{
+				printf("\nProvided values didn't match!\n");
+			}
+		}
+
+		return Credential;
+	}
+
+	bool SupportsPersistentCredentials()
+	{
+#if PLATFORM_WINDOWS
+		return true;
+#else
+		return false;
+#endif
+	}
+
+	bool SaveCredential(const FCredential& Credential)
+	{
+#if PLATFORM_WINDOWS
+		const FString NamespacedTargetName =
+			FString::Printf(TEXT("SwitchboardListener_%s"), *Credential.CredentialName);
+
+		CREDENTIALW Cred = { 0 };
+		Cred.Type = CRED_TYPE_GENERIC;
+		Cred.TargetName = const_cast<FString::ElementType*>(*NamespacedTargetName);
+		Cred.CredentialBlobSize = sizeof(TCHAR) * Credential.CredentialBlob.Len();
+		Cred.CredentialBlob = reinterpret_cast<BYTE*>(
+			const_cast<FString::ElementType*>(*Credential.CredentialBlob));
+		Cred.Persist = CRED_PERSIST_LOCAL_MACHINE;
+		Cred.UserName = const_cast<FString::ElementType*>(*Credential.UserName);
+		if (::CredWriteW(&Cred, 0))
+		{
+			return true;
+		}
+		else
+		{
+			UE_LOGFMT(LogSwitchboard, Warning, "CredWriteW failed; GetLastError = {LastError}", static_cast<int64>(::GetLastError()));
+		}
+#else
+		// Not implemented
+#endif
+
+		return false;
+	}
+
+	TOptional<FCredential> LoadCredential(FStringView CredentialName)
+	{
+#if PLATFORM_WINDOWS
+		const FString NamespacedTargetName =
+			FString::Printf(TEXT("SwitchboardListener_%*s"),
+				CredentialName.Len(), CredentialName.GetData());
+
+		CREDENTIALW* OutCredPtr = nullptr;
+		if (::CredReadW(*NamespacedTargetName, CRED_TYPE_GENERIC, 0, &OutCredPtr))
+		{
+			FCredential Result;
+			Result.CredentialName = CredentialName;
+			Result.UserName = FString(OutCredPtr->UserName);
+			Result.CredentialBlob = FString(
+				OutCredPtr->CredentialBlobSize / sizeof(TCHAR),
+				reinterpret_cast<TCHAR*>(OutCredPtr->CredentialBlob));
+
+			::CredFree(OutCredPtr);
+
+			return Result;
+		}
+		else
+		{
+			const int64 LastError = ::GetLastError();
+			if (LastError == ERROR_NOT_FOUND)
+			{
+				UE_LOGFMT(LogSwitchboard, Verbose, "CredReadW returned ERROR_NOT_FOUND");
+			}
+			else
+			{
+				UE_LOGFMT(LogSwitchboard, Warning, "CredReadW failed; GetLastError = {LastError}", static_cast<int64>(::GetLastError()));
+			}
+		}
+#else
+		// Not implemented
+#endif
+
+		return {};
+	}
+
+
+	void FillSecureRandom(TArray<uint8>& InOutArray)
+	{
+		const int32 NumBytes = InOutArray.Num();
+		if (ensure(NumBytes > 0))
+		{
+			check(1 == RAND_bytes(
+				reinterpret_cast<unsigned char*>(InOutArray.GetData()),
+				NumBytes));
+		}
+	}
+}
+
+
+namespace UE::SwitchboardListener::Certificates::Private
+{
+	/**
+	 * Generate a 2048-bit RSA key.
+	 *
+	 * @param PrivateKey The private key struct
+	 */
+	EVP_PKEY* GeneratePrivateKey()
+	{
+		EVP_PKEY* PrivateKey = EVP_PKEY_new();
+
+		// Generate the RSA key and assign it to pkey.
+		RSA* Rsa = RSA_new();
+
+		{
+			BIGNUM* BigNum = BN_new();
+			BN_set_word(BigNum, RSA_F4);
+			RSA_generate_key_ex(Rsa, 2048, BigNum, nullptr);
+			BN_free(BigNum);
+		}
+
+		// EVP_PKEY_assign_RSA will "use the supplied (Rsa) internally
+		// and so (Rsa) will be freed when the parent pkey is freed."
+		if (!EVP_PKEY_assign_RSA(PrivateKey, Rsa))
+		{
+			EVP_PKEY_free(PrivateKey);
+			return nullptr;
+		}
+
+		return PrivateKey;
+	}
+
+
+	FString BIOToString(BIO* Bio)
+	{
+		char* Text;
+		long Length = BIO_get_mem_data(Bio, &Text);
+
+		auto Convert = StringCast<TCHAR>(Text, Length);
+		FString Result(Convert.Length(), Convert.Get());
+		Result.AppendChar(TEXT('\0'));
+		return Result;
+	}
+
+
+	/**
+	 * Convert a certificate to a string.
+	 *
+	 * @param x509 The certificate struct
+	 */
+	FString X509ToString(X509* x509)
+	{
+		BIO* Bio = BIO_new(BIO_s_mem());
+		ON_SCOPE_EXIT{ BIO_free(Bio); };
+
+		check(1 == PEM_write_bio_X509(Bio, x509));
+
+		return BIOToString(Bio);
+	}
+
+
+	/**
+	 * Converts a private key to a string.
+	 *
+	 * @param PrivateKey The private key struct
+	 */
+	FString PrivateKeyToString(EVP_PKEY* PrivateKey, const FString& PrivateKeyPassword)
+	{
+		BIO* Bio = BIO_new(BIO_s_mem());
+		ON_SCOPE_EXIT{ BIO_free(Bio); };
+
+		if (PrivateKeyPassword.IsEmpty())
+		{
+			check(1 == PEM_write_bio_PrivateKey_traditional(Bio, PrivateKey,
+				nullptr,
+				nullptr, 0, nullptr,
+				nullptr));
+		}
+		else
+		{
+			check(1 == PEM_write_bio_PrivateKey_traditional(Bio, PrivateKey,
+				EVP_des_ede3_cbc(),
+				nullptr, 0, nullptr,
+				TCHAR_TO_UTF8(*PrivateKeyPassword)));
+		}
+
+		return BIOToString(Bio);
+	}
+
+
+	/**
+	 * Generate self-signed x509 certificate.
+	 *
+	 * @param PrivateKey The private key struct
+	 */
+	X509* GenerateCertificate(EVP_PKEY* PrivateKey)
+	{
+		X509* x509 = X509_new();
+
+		// Set the serial number.
+		uint64_t SerialNumber;
+		check(1 == RAND_bytes(reinterpret_cast<unsigned char*>(&SerialNumber), sizeof(SerialNumber)));
+		ASN1_INTEGER_set_uint64(X509_get_serialNumber(x509), SerialNumber);
+
+		// This certificate is valid for 10 years.
+		X509_gmtime_adj(X509_get_notBefore(x509), 0);
+		X509_gmtime_adj(X509_get_notAfter(x509), (60L*60*24*365*10));
+
+		// Set the public key for our certificate.
+		X509_set_pubkey(x509, PrivateKey);
+
+		// Set the common name, and copy the subject name to the issuer name.
+		X509_NAME* SubjectName = X509_get_subject_name(x509);
+		X509_NAME_add_entry_by_txt(SubjectName, "CN", MBSTRING_ASC,
+			(unsigned char*)"SwitchboardListener", -1, -1, 0);
+		X509_set_issuer_name(x509, SubjectName);
+
+		// Actually sign the certificate with our key.
+		if (!X509_sign(x509, PrivateKey, EVP_sha256()))
+		{
+			X509_free(x509);
+			return nullptr;
+		}
+
+		return x509;
+	}
+
+#if !PLATFORM_WINDOWS
+	bool EnsurePathOnlyOwnerReadable(const FString& InPath)
+	{
+		FString NormalizedPath = InPath;
+		FPaths::NormalizeFilename(NormalizedPath);
+		NormalizedPath = FPaths::ConvertRelativePathToFull(NormalizedPath);
+
+		struct stat Stat;
+		if (stat(TCHAR_TO_UTF8(*NormalizedPath), &Stat) == 0)
+		{
+			// Clear group/other read/write/execute bits
+			Stat.st_mode &= ~(S_IRWXG | S_IRWXO);
+			chmod(TCHAR_TO_UTF8(*NormalizedPath), Stat.st_mode);
+			return true;
+		}
+		else
+		{
+			return false;
+		}
+	}
+#endif
+}
+
+
+namespace UE::SwitchboardListener::Certificates
+{
+	using namespace UE::SwitchboardListener::Certificates::Private;
+
+
+	TTuple<FString, FString> GetSelfSignedPaths(bool bCreate /* = false */)
+	{
+		IPlatformFile& FileManager = FPlatformFileManager::Get().GetPlatformFile();
+
+#if PLATFORM_WINDOWS
+		wchar_t* LocalAppDataWstr = nullptr;
+		::SHGetKnownFolderPath(FOLDERID_LocalAppData, 0, nullptr, &LocalAppDataWstr);
+		ON_SCOPE_EXIT{ ::CoTaskMemFree(LocalAppDataWstr); };
+
+		const FString LocalAppDataDir(LocalAppDataWstr);
+		const FString SblHomeDir = LocalAppDataDir / TEXT("SwitchboardListener");
+		if (bCreate)
+		{
+			FileManager.CreateDirectory(*SblHomeDir);
+		}
+
+		const FString CertificateFile = SblHomeDir / TEXT("transportcert.pem");
+		const FString PrivateKeyFile = SblHomeDir / TEXT("transportkey.key");
+#else
+		const FString UserHomeDir = FPlatformProcess::UserHomeDir();
+		const FString SblHomeDir = UserHomeDir / TEXT(".switchboardlistener");
+		const FString SblCertDir = SblHomeDir / TEXT("cert");
+
+		if (bCreate)
+		{
+			FileManager.CreateDirectory(*SblHomeDir);
+			FileManager.CreateDirectory(*SblCertDir);
+			EnsurePathOnlyOwnerReadable(SblCertDir);
+		}
+
+		const FString CertificateFile = SblCertDir / TEXT("transportcert.pem");
+		const FString PrivateKeyFile = SblCertDir / TEXT("transportkey.key");
+#endif
+
+		return TTuple<FString, FString>(CertificateFile, PrivateKeyFile);
+	}
+
+
+	TOptional<FString> CreateSelfSigned(const FString& PrivateKeyPassword)
+	{
+		EVP_PKEY* PrivateKey = GeneratePrivateKey();
+		if (!PrivateKey)
+		{
+			return TOptional<FString>{};
+		}
+
+		ON_SCOPE_EXIT{ EVP_PKEY_free(PrivateKey); };
+
+		X509* x509 = GenerateCertificate(PrivateKey);
+		if (!x509)
+		{
+			return TOptional<FString>{};
+		}
+
+		ON_SCOPE_EXIT{ X509_free(x509); };
+
+		// Compute SHA256 fingerprint.
+		uint8 Digest[EVP_MAX_MD_SIZE];
+		uint32 DigestSize = 0;
+		check(1 == X509_digest(x509, EVP_sha256(), Digest, &DigestSize));
+
+		// Return the first half in the format "09:AB:CD:..."
+		FString Fingerprint;
+		for (uint32 ByteIdx = 0; ByteIdx < DigestSize/2; ++ByteIdx)
+		{
+			ByteToHex(Digest[ByteIdx], Fingerprint);
+			if (ByteIdx < (DigestSize/2 - 1))
+			{
+				Fingerprint.AppendChar(TEXT(':'));
+			}
+		}
+
+		// Get PEM strings and write to files.
+		const FString PrivateKeyString = PrivateKeyToString(PrivateKey, PrivateKeyPassword);
+		const FString CertificateString = X509ToString(x509);
+
+		TTuple<FString, FString> Paths = GetSelfSignedPaths();
+		const FString CertificateFile = Paths.Get<0>();
+		const FString PrivateKeyFile = Paths.Get<1>();
+
+		bool bSucceeded = true;
+		bSucceeded &= FFileHelper::SaveStringToFile(CertificateString, *CertificateFile);
+		bSucceeded &= FFileHelper::SaveStringToFile(PrivateKeyString, *PrivateKeyFile);
+
+#if !PLATFORM_WINDOWS
+		EnsurePathOnlyOwnerReadable(PrivateKeyFile);
+#endif
+
+		return bSucceeded ? Fingerprint : TOptional<FString>{};
+	}
+};
diff --git a/Engine/Source/Programs/SwitchboardListener/SblCore/Private/SwitchboardAuth.h b/Engine/Source/Programs/SwitchboardListener/SblCore/Private/SwitchboardAuth.h
new file mode 100644
--- /dev/null
+++ b/Engine/Source/Programs/SwitchboardListener/SblCore/Private/SwitchboardAuth.h
@@ -0,0 +1,47 @@
+// Copyright Epic Games, Inc. All Rights Reserved.
+
+#pragma once
+
+#include "Containers/UnrealString.h"
+#include "Templates/Tuple.h"
+
+
+namespace UE::SwitchboardListener
+{
+	FString ReadCredentialFromStdin(const FString& CredentialDesc, bool bAllowEmpty = false);
+
+	struct FCredential
+	{
+		FString CredentialName;
+		FString UserName;
+		FString CredentialBlob;
+	};
+
+	bool SupportsPersistentCredentials();
+	bool SaveCredential(const FCredential& Credential);
+	TOptional<FCredential> LoadCredential(FStringView CredentialName);
+
+	void FillSecureRandom(TArray<uint8>& InOutArray);
+};
+
+
+namespace UE::SwitchboardListener::Certificates
+{
+	/**
+	 * Get the default paths for (self-signed certificate, private key).
+	 * On Windows, %LocalAppData%\SwitchboardListener
+	 * Otherwise, ~/.switchboardlistener/cert, NOT world-readable.
+	 * 
+	 * @param bCreate If true, and the paths do not exist, create them with proper permissions.
+	 */
+	TTuple<FString, FString> GetSelfSignedPaths(bool bCreate = false);
+
+	/**
+	 * Generate a private key + self-signed certificate, and write them to the default paths.
+	 * 
+	 * @param PrivateKeyPassword The password with which to encrypt the PEM private key (unencrypted if empty).
+	 * @return The truncated (first 128 bits) of the SHA256 cert fingerprint for visual comparison.
+	 *         If the optional is unset, an error occurred during generation or save.
+	 */
+	TOptional<FString> CreateSelfSigned(const FString& PrivateKeyPassword);
+};
diff --git a/Engine/Source/Programs/SwitchboardListener/SblCore/Private/SwitchboardListener.cpp b/Engine/Source/Programs/SwitchboardListener/SblCore/Private/SwitchboardListener.cpp
new file mode 100644
--- /dev/null
+++ b/Engine/Source/Programs/SwitchboardListener/SblCore/Private/SwitchboardListener.cpp
@@ -0,0 +1,3015 @@
+// Copyright Epic Games, Inc. All Rights Reserved.
+
+#include "SwitchboardListener.h"
+
+#include "CpuUtilizationMonitor.h"
+#include "SBLHelperClient.h"
+#include "SwitchboardAuth.h"
+#include "SwitchboardListenerApp.h"
+#include "SwitchboardMessageFuture.h"
+#include "SwitchboardPacket.h"
+#include "SwitchboardProtocol.h"
+#include "SwitchboardTasks.h"
+#include "SyncStatus.h"
+
+#include "Algo/Find.h"
+#include "Async/Async.h"
+#include "Async/AsyncWork.h"
+#include "Common/TcpListener.h"
+#include "GenericPlatform/GenericPlatformMisc.h"
+#include "GenericPlatform/GenericPlatformProcess.h"
+#include "HAL/FileManager.h"
+#include "HAL/PlatformFileManager.h"
+#include "Interfaces/IPv4/IPv4Endpoint.h"
+#include "IPAddress.h"
+#include "Logging/StructuredLog.h"
+#include "Misc/Base64.h"
+#include "Misc/CommandLine.h"
+#include "Misc/FileHelper.h"
+#include "Misc/Paths.h"
+#include "Misc/ScopeExit.h"
+#include "String/BytesToHex.h"
+#include "String/HexToBytes.h"
+#include <atomic>
+
+#if !PLATFORM_WINDOWS
+#include <string.h> // memset
+#endif
+
+
+#if PLATFORM_WINDOWS
+
+#pragma warning(push)
+#pragma warning(disable : 4005)	// Disable macro redefinition warning for compatibility with Windows SDK 8+
+
+#include "Windows/AllowWindowsPlatformTypes.h"
+#include <Windows.h>
+#include <WinUser.h>
+#include <shellapi.h>
+#include "nvapi.h"
+#include "Windows/HideWindowsPlatformTypes.h"
+
+#pragma warning(pop)
+
+static void FillOutMosaicTopologies(TArray<FMosaicTopo>& MosaicTopos);
+
+#endif
+
+
+#define QUIC_ENSURE(X)		ensure(QUIC_SUCCEEDED(X))
+
+
+const FIPv4Endpoint FSwitchboardListener::InvalidEndpoint(FIPv4Address::LanBroadcast, 0);
+
+
+namespace
+{
+#if !UE_BUILD_DEBUG
+	const double DefaultInactiveTimeoutSeconds = 5.0;
+#else
+	// Infinite timeout in debug; convenient for pausing at breakpoints for long periods.
+	const double DefaultInactiveTimeoutSeconds = 0.0;
+#endif
+
+	bool TryFindIdInBrokenMessage(const FString& InMessage, FGuid& OutGuid)
+	{
+		// Try to parse message that could not be parsed regularly and look for the message ID.
+		// This way we can at least tell Switchboard which message was broken.
+		const int32 IdIdx = InMessage.Find(TEXT("\"id\""));
+		if (IdIdx > 0)
+		{
+			const FString Chopped = InMessage.RightChop(IdIdx);
+			FString LeftOfComma;
+			FString RightOfComma;
+			if (Chopped.Split(",", &LeftOfComma, &RightOfComma))
+			{
+				FString LeftOfColon;
+				FString RightOfColon;
+				if (LeftOfComma.Split(":", &LeftOfColon, &RightOfColon))
+				{
+					RightOfColon.TrimStartAndEndInline();
+					bool bDoubleQuotesRemoved = false;
+					RightOfColon.TrimQuotesInline(&bDoubleQuotesRemoved);
+					if (!bDoubleQuotesRemoved)
+					{
+						// remove single quotes if there were no double quotes
+						RightOfColon.LeftChopInline(1);
+						RightOfColon.RightChopInline(1);
+					}
+
+					return FGuid::Parse(RightOfColon, OutGuid);
+				}
+			}
+		}
+
+		return false;
+	}
+
+	QUIC_ADDR QuicAddrFromEndpoint(FIPv4Endpoint Endpoint)
+	{
+		QUIC_ADDR QuicAddr = {};
+
+		inet_pton(AF_INET, TCHAR_TO_ANSI(*Endpoint.Address.ToString()),
+			&QuicAddr.Ipv4.sin_addr.s_addr);
+
+		QuicAddrSetFamily(&QuicAddr, QUIC_ADDRESS_FAMILY_INET);
+		QuicAddrSetPort(&QuicAddr, Endpoint.Port);
+
+		return QuicAddr;
+	}
+
+
+	FIPv4Endpoint EndpointFromQuicAddr(QUIC_ADDR QuicAddr)
+	{
+		FIPv4Endpoint Endpoint;
+
+		char Ipv4Buf[INET_ADDRSTRLEN];
+		inet_ntop(AF_INET, &QuicAddr.Ipv4.sin_addr.s_addr, Ipv4Buf, sizeof(Ipv4Buf));
+		FIPv4Address::Parse(ANSI_TO_TCHAR(Ipv4Buf), Endpoint.Address);
+
+		Endpoint.Port = QuicAddrGetPort(&QuicAddr);
+		return Endpoint;
+	}
+
+	void SecureZero(void* Buffer, size_t Length)
+	{
+#if PLATFORM_WINDOWS
+		SecureZeroMemory(Buffer, Length);
+#else
+		static void* (* const volatile memset_ptr)(void*, int, size_t) = memset;
+		(memset_ptr)(Buffer, 0, Length);
+#endif
+	}
+}
+
+struct FRunningProcess
+{
+	uint32 PID = 0;
+	FGuid UUID;
+	FProcHandle Handle;
+
+	void* StdoutParentReadPipe = nullptr;
+	void* StdoutChildWritePipe = nullptr;
+	void* StdinChildReadPipe = nullptr;
+	void* StdinParentWritePipe = nullptr;
+
+	TArray<uint8> Output;
+
+	FIPv4Endpoint Recipient;
+	FString Path;
+	FString Name;
+	FString Caller;
+
+	std::atomic<bool> bPendingKill = false;
+	bool bUpdateClientsWithStdout = false;
+	bool bLockGpuClock = false;
+
+	bool CreatePipes()
+	{
+		if (!FPlatformProcess::CreatePipe(StdoutParentReadPipe, StdoutChildWritePipe))
+		{
+			UE_LOG(LogSwitchboard, Error, TEXT("Could not create pipe to read process output!"));
+			return false;
+		}
+
+		if (!FPlatformProcess::CreatePipe(StdinChildReadPipe, StdinParentWritePipe, true))
+		{
+			UE_LOG(LogSwitchboard, Error, TEXT("Could not create pipe to write process input!"));
+			return false;
+		}
+
+		return true;
+	}
+
+	void ClosePipes()
+	{
+		FPlatformProcess::ClosePipe(StdoutParentReadPipe, StdoutChildWritePipe);
+		FPlatformProcess::ClosePipe(StdinChildReadPipe, StdinParentWritePipe);
+	}
+};
+
+FSwitchboardCommandLineOptions FSwitchboardCommandLineOptions::FromString(const TCHAR* CommandLine)
+{
+	FSwitchboardCommandLineOptions OutOptions;
+
+	TArray<FString> Tokens;
+	TArray<FString> Switches;
+	FCommandLine::Parse(CommandLine, Tokens, Switches);
+	TMap<FString, FString> SwitchPairs;
+	for (int32 SwitchIdx = Switches.Num() - 1; SwitchIdx >= 0; --SwitchIdx)
+	{
+		FString& Switch = Switches[SwitchIdx];
+		TArray<FString> SplitSwitch;
+		if (2 == Switch.ParseIntoArray(SplitSwitch, TEXT("="), true))
+		{
+			SwitchPairs.Add(SplitSwitch[0], SplitSwitch[1].TrimQuotes());
+			Switches.RemoveAt(SwitchIdx);
+		}
+	}
+
+	if (Switches.Contains(TEXT("version")))
+	{
+		OutOptions.OutputVersion = true;
+	}
+
+	if (Switches.Contains(TEXT("noMinimizeOnLaunch")))
+	{
+		OutOptions.MinimizeOnLaunch = false;
+	}
+
+	if (SwitchPairs.Contains(TEXT("ip")))
+	{
+		FIPv4Address ParseAddr;
+		if (FIPv4Address::Parse(SwitchPairs[TEXT("ip")], ParseAddr))
+		{
+			OutOptions.Address = ParseAddr;
+		}
+	}
+
+	OutOptions.SecureMode = ESecureMode::Unspecified;
+
+	if (SwitchPairs.Contains(TEXT("certfile")))
+	{
+		OutOptions.CertificateFile = SwitchPairs[TEXT("certfile")];
+		OutOptions.SecureMode = ESecureMode::CertificateFromFile;
+	}
+
+	if (SwitchPairs.Contains(TEXT("keyfile")))
+	{
+		OutOptions.PrivateKeyFile = SwitchPairs[TEXT("keyfile")];
+		OutOptions.SecureMode = ESecureMode::CertificateFromFile;
+	}
+
+	if (SwitchPairs.Contains(TEXT("certhash")))
+	{
+#if PLATFORM_WINDOWS
+		OutOptions.CertificateHash = SwitchPairs[TEXT("certhash")];
+
+		if (OutOptions.SecureMode == ESecureMode::CertificateFromFile)
+		{
+			UE_LOGFMT(LogSwitchboard, Error, "Both certificate file and hash arguments were specified; using certificate hash, IGNORING file");
+		}
+
+		OutOptions.SecureMode = ESecureMode::CertificateByHash;
+#else
+		UE_LOGFMT(LogSwitchboard, Fatal, "Certificate by hash is only supported on Windows");
+#endif
+	}
+
+	if (SwitchPairs.Contains(TEXT("port")))
+	{
+		OutOptions.Port = FCString::Atoi(*SwitchPairs[TEXT("port")]);
+	}
+
+	if (SwitchPairs.Contains(TEXT("RedeployFromPid")))
+	{
+		OutOptions.RedeployFromPid = FCString::Atoi64(*SwitchPairs[TEXT("RedeployFromPid")]);
+	}
+
+	return OutOptions;
+}
+
+FString FSwitchboardCommandLineOptions::ToString(bool bIncludeRedeploy /* = false */) const
+{
+	TArray<FString> Args;
+
+	if (!MinimizeOnLaunch)
+	{
+		Args.Add(TEXT("-noMinimizeOnLaunch"));
+	}
+
+	if (Address.IsSet())
+	{
+		Args.Add(FString::Printf(TEXT("-ip=%s"), *Address.GetValue().ToString()));
+	}
+
+	if (Port.IsSet())
+	{
+		Args.Add(FString::Printf(TEXT("-port=%u"), Port.GetValue()));
+	}
+
+	if (bIncludeRedeploy)
+	{
+		if (RedeployFromPid.IsSet())
+		{
+			Args.Add(FString::Printf(TEXT("-RedeployFromPid=%u"), RedeployFromPid.GetValue()));
+		}
+	}
+
+	return FString::Join(Args, TEXT(" "));
+}
+
+FSwitchboardListener::FSwitchboardListener(const FSwitchboardCommandLineOptions& InOptions)
+	: Options(InOptions)
+	, CpuMonitor(MakeShared<FCpuUtilizationMonitor>())
+	, SBLHelper(MakeShared<FSBLHelperClient>())
+	, bIsNvAPIInitialized(false)
+	, CachedMosaicToposLock(MakeShared<FRWLock>())
+	, CachedMosaicTopos(MakeShared<TArray<FMosaicTopo>>())
+{
+#if PLATFORM_WINDOWS
+	// initialize NvAPI
+	{
+		const NvAPI_Status Result = NvAPI_Initialize();
+		if (Result == NVAPI_OK)
+		{
+			bIsNvAPIInitialized = true;
+
+			FillOutMosaicTopologies(*CachedMosaicTopos);
+		}
+		else
+		{
+			NvAPI_ShortString ErrorString;
+			NvAPI_GetErrorMessage(Result, ErrorString);
+			UE_LOG(LogSwitchboard, Error, TEXT("NvAPI_Initialize failed. Error: %s"), ANSI_TO_TCHAR(ErrorString));
+		}
+	}
+#endif // PLATFORM_WINDOWS
+
+	const FIPv4Address DefaultIp = FIPv4Address(0, 0, 0, 0);
+	if (!Options.Address.IsSet())
+	{
+		UE_LOG(LogSwitchboard, Warning, TEXT("Defaulting to: -ip=%s"), *DefaultIp.ToString());
+	}
+
+	const uint16 DefaultPort = 2980;
+	if (!Options.Port.IsSet())
+	{
+		UE_LOG(LogSwitchboard, Warning, TEXT("Defaulting to: -port=%u"), DefaultPort);
+	}
+
+	ListenerEndpoint = MakeUnique<FIPv4Endpoint>(Options.Address.Get(DefaultIp), Options.Port.Get(DefaultPort));
+}
+
+FSwitchboardListener::~FSwitchboardListener()
+{
+}
+
+bool FSwitchboardListener::Init(bool bStartListeningImmediately /* = true */)
+{
+	if (Options.SecureMode == FSwitchboardCommandLineOptions::ESecureMode::Unspecified)
+	{
+		TTuple<FString, FString> CertKeyPaths = UE::SwitchboardListener::Certificates::GetSelfSignedPaths();
+
+		Options.SecureMode = FSwitchboardCommandLineOptions::ESecureMode::CertificateFromFile;
+		Options.CertificateFile = CertKeyPaths.Get<0>();
+		Options.PrivateKeyFile = CertKeyPaths.Get<1>();
+
+		// Try to load the private key password.
+		const FString PrivateKeyPwCredentialName = FString::Printf(
+			TEXT("PrivateKeyPassword_%s"), **Options.PrivateKeyFile);
+		if (TOptional<UE::SwitchboardListener::FCredential> PrivateKeyPwCredential =
+			UE::SwitchboardListener::LoadCredential(PrivateKeyPwCredentialName))
+		{
+			PrivateKeyPassword = PrivateKeyPwCredential->CredentialBlob;
+		}
+
+
+		if (!FPaths::FileExists(*Options.CertificateFile)
+			|| !FPaths::FileExists(*Options.PrivateKeyFile)
+			|| (!PrivateKeyPassword || PrivateKeyPassword->IsEmpty()))
+		{
+			UE_LOGFMT(LogSwitchboard, Display, "SwitchboardListener requires a TLS certificate to be specified; generating self-signed certificate...");
+
+			// Platforms with a persistent credential store encrypt the generated
+			// private key with a random string, and persist it for the user.
+			if (UE::SwitchboardListener::SupportsPersistentCredentials())
+			{
+				UE_LOGFMT(LogSwitchboard, Verbose, "Generating random private key password");
+				constexpr int32 RandomPwByteLength = 20;
+				TArray<uint8> RandomPwBytes;
+				RandomPwBytes.SetNumUninitialized(RandomPwByteLength);
+				UE::SwitchboardListener::FillSecureRandom(RandomPwBytes);
+				PrivateKeyPassword = BytesToHexLower(RandomPwBytes.GetData(), RandomPwByteLength);
+
+				UE::SwitchboardListener::FCredential PrivateKeyPwCredential = {
+					.CredentialName = PrivateKeyPwCredentialName,
+					.CredentialBlob = *PrivateKeyPassword,
+				};
+
+				check(UE::SwitchboardListener::SaveCredential(PrivateKeyPwCredential));
+			}
+
+			TOptional<FString> Fingerprint =
+				UE::SwitchboardListener::Certificates::CreateSelfSigned(
+					PrivateKeyPassword.Get(FString()));
+
+			check(Fingerprint.IsSet());
+
+			UE_LOGFMT(LogSwitchboard, Display, "Your new self-signed certificate fingerprint is: {Fingerprint}", *Fingerprint);
+		}
+	}
+
+	if (TOptional<UE::SwitchboardListener::FCredential> SavedCredential =
+		UE::SwitchboardListener::LoadCredential(PasswordCredentialName))
+	{
+		UE_LOGFMT(LogSwitchboard, Display, "Using stored password");
+		ExpectedAuthPassword = SavedCredential->CredentialBlob;
+	}
+
+	OnInitDelegate.Broadcast();
+
+	if (bStartListeningImmediately)
+	{
+		return StartListening();
+	}
+	else
+	{
+		return true;
+	}
+}
+
+
+void FSwitchboardListener::Shutdown()
+{
+	OnShutdownDelegate.Broadcast();
+}
+
+
+bool FSwitchboardListener::SetAuthPassword(const FString& NewPassword)
+{
+	if (NewPassword.IsEmpty())
+	{
+		return false;
+	}
+
+	ExpectedAuthPassword = NewPassword;
+
+	if (UE::SwitchboardListener::SupportsPersistentCredentials())
+	{
+		UE::SwitchboardListener::FCredential NewCredential = {
+			.CredentialName = PasswordCredentialName,
+			.CredentialBlob = ExpectedAuthPassword,
+		};
+
+		UE_LOGFMT(LogSwitchboard, Display, "Storing provided password");
+		ensure(UE::SwitchboardListener::SaveCredential(NewCredential));
+	}
+
+	return true;
+}
+
+
+TSet<FIPv4Endpoint> FSwitchboardListener::GetConnectedClientEndpoints() const
+{
+	TSet<FIPv4Endpoint> OutEndpointSet;
+	ConnectionsByEndpoint.GetKeys(OutEndpointSet);
+	return OutEndpointSet;
+}
+
+
+TSet<FIPv4Address> FSwitchboardListener::GetConnectedClientAddresses() const
+{
+	TSet<FIPv4Address> OutAddressSet;
+	for (const TPair<FIPv4Endpoint, FConnectionRef>& Endpoint : ConnectionsByEndpoint)
+	{
+		OutAddressSet.Add(Endpoint.Key.Address);
+	}
+	return OutAddressSet;
+}
+
+
+bool FSwitchboardListener::StartListening()
+{
+	QUIC_STATUS Status;
+	if (QUIC_FAILED(Status = MsQuicOpen2(&QuicApi)))
+	{
+		UE_LOGFMT(LogSwitchboard, Error, "MsQuicOpen2 failed with status {Status}", static_cast<int64>(Status));
+		return false;
+	}
+
+	// Create a registration for the app's connections. This sets a name for the
+	// app (used for persistent storage and for debugging). It also configures
+	// the execution profile, using the default "low latency" profile.
+	const QUIC_REGISTRATION_CONFIG RegConfig = { "switchboardlistener", QUIC_EXECUTION_PROFILE_LOW_LATENCY };
+	if (QUIC_FAILED(Status = QuicApi->RegistrationOpen(&RegConfig, &QuicRegistration))) {
+		UE_LOGFMT(LogSwitchboard, Error, "MsQuic RegistrationOpen failed with status {Status}", static_cast<int64>(Status));
+		return false;
+	}
+
+	QUIC_CERTIFICATE_FILE CertFile = { 0 };
+	QUIC_CERTIFICATE_FILE_PROTECTED CertFileProtected = { 0 };
+	QUIC_CERTIFICATE_HASH CertHash = { 0 };
+	QUIC_CREDENTIAL_CONFIG CredConfig = {
+		.Type = QUIC_CREDENTIAL_TYPE_NONE,
+		.Flags =
+			QUIC_CREDENTIAL_FLAG_INDICATE_CERTIFICATE_RECEIVED |
+			QUIC_CREDENTIAL_FLAG_USE_PORTABLE_CERTIFICATES,
+	};
+
+	// Path FString -> UTF8 backing stores for the ConfigurationLoadCredential call.
+	TArray<UTF8CHAR> CertificateFileArray, PrivateKeyFileArray, PrivateKeyPwArray;
+
+	if (Options.SecureMode == FSwitchboardCommandLineOptions::ESecureMode::CertificateFromFile)
+	{
+		auto PrivateKeyFileU8 = StringCast<UTF8CHAR>(**Options.PrivateKeyFile);
+		PrivateKeyFileArray = TArray<UTF8CHAR>(PrivateKeyFileU8.Get(), PrivateKeyFileU8.Length() + 1);
+		auto CertificateFileU8 = StringCast<UTF8CHAR>(**Options.CertificateFile);
+		CertificateFileArray = TArray<UTF8CHAR>(CertificateFileU8.Get(), CertificateFileU8.Length() + 1);
+
+		if (PrivateKeyPassword)
+		{
+			auto PrivateKeyPwU8 = StringCast<UTF8CHAR>(**PrivateKeyPassword);
+			PrivateKeyPwArray = TArray<UTF8CHAR>(PrivateKeyPwU8.Get(), PrivateKeyPwU8.Length() + 1);
+			SecureZero(const_cast<UTF8CHAR*>(PrivateKeyPwU8.Get()), PrivateKeyPwU8.Length());
+
+			CertFileProtected.PrivateKeyFile = reinterpret_cast<const char*>(PrivateKeyFileArray.GetData());
+			CertFileProtected.CertificateFile = reinterpret_cast<const char*>(CertificateFileArray.GetData());
+			CertFileProtected.PrivateKeyPassword = reinterpret_cast<const char*>(PrivateKeyPwArray.GetData());
+
+			CredConfig.Type = QUIC_CREDENTIAL_TYPE_CERTIFICATE_FILE_PROTECTED;
+			CredConfig.CertificateFileProtected = &CertFileProtected;
+		}
+		else
+		{
+			CertFile.PrivateKeyFile = reinterpret_cast<const char*>(PrivateKeyFileArray.GetData());
+			CertFile.CertificateFile = reinterpret_cast<const char*>(CertificateFileArray.GetData());
+
+			CredConfig.Type = QUIC_CREDENTIAL_TYPE_CERTIFICATE_FILE;
+			CredConfig.CertificateFile = &CertFile;
+		}
+	}
+#if PLATFORM_WINDOWS
+	else if (Options.SecureMode == FSwitchboardCommandLineOptions::ESecureMode::CertificateByHash)
+	{
+		constexpr int32 ExpectedBytes = sizeof(CertHash.ShaHash);
+		constexpr int32 ExpectedStrLen = ExpectedBytes * 2;
+		const int32 HashStrLen = Options.CertificateHash->Len();
+		if (HashStrLen != ExpectedStrLen)
+		{
+			UE_LOGFMT(LogSwitchboard, Error, "Certificate hash should be {ExpectedStrLen} characters ({ExpectedBytes} bytes); got {HashStrLen} characters",
+				ExpectedStrLen, ExpectedBytes, HashStrLen);
+			return false;
+		}
+
+		const int32 HashLen = UE::String::HexToBytes(Options.CertificateHash.GetValue(), CertHash.ShaHash);
+		ensure(HashLen == ExpectedBytes);
+
+		UE_LOGFMT(LogSwitchboard, Display, "Using certificate with hash {CertificateHash}", Options.CertificateHash.GetValue());
+		CredConfig.Type = QUIC_CREDENTIAL_TYPE_CERTIFICATE_HASH;
+		CredConfig.CertificateHash = &CertHash;
+	}
+#endif // PLATFORM_WINDOWS
+	else
+	{
+		// This should never happen.
+		checkf(false, TEXT("Unsupported credential configuration"));
+		return false;
+	}
+
+	QUIC_SETTINGS Settings = {};
+
+	Settings.IdleTimeoutMs = DefaultInactiveTimeoutSeconds * 1000.0;
+	Settings.IsSet.IdleTimeoutMs = 1;
+
+	// Allow each peer to open a single bidirectional stream.
+	Settings.PeerBidiStreamCount = 1;
+	Settings.IsSet.PeerBidiStreamCount = 1;
+
+	// The protocol name used in Application Layer Protocol Negotiation.
+	constexpr char SblAlpnStr[] = "ue-switchboard";
+	const QUIC_BUFFER SblAlpn = { sizeof(SblAlpnStr) - 1, (uint8_t*)SblAlpnStr };
+
+	// Allocate/initialize the configuration object with the configured ALPN and settings.
+	if (QUIC_FAILED(Status = QuicApi->ConfigurationOpen(QuicRegistration, &SblAlpn, 1, &Settings, sizeof(Settings), NULL, &QuicConfiguration)))
+	{
+		UE_LOGFMT(LogSwitchboard, Error, "MsQuic ConfigurationOpen failed with status {Status}", static_cast<int64>(Status));
+		return false;
+	}
+
+	// Loads the TLS credential part of the configuration.
+	if (QUIC_FAILED(Status = QuicApi->ConfigurationLoadCredential(QuicConfiguration, &CredConfig)))
+	{
+		UE_LOGFMT(LogSwitchboard, Error, "MsQuic ConfigurationLoadCredential failed with status {Status}", static_cast<int64>(Status));
+		return false;
+	}
+	else
+	{
+		if (PrivateKeyPwArray.Num())
+		{
+			SecureZero(PrivateKeyPwArray.GetData(), PrivateKeyPwArray.Num());
+		}
+
+		if (PrivateKeyPassword.IsSet())
+		{
+			SecureZero(PrivateKeyPassword->GetCharArray().GetData(),
+				PrivateKeyPassword->GetCharArray().Num() * sizeof(TCHAR));
+			PrivateKeyPassword.Reset();
+		}
+	}
+
+	// Create/allocate a new listener object.
+	if (QUIC_FAILED(Status = QuicApi->ListenerOpen(QuicRegistration, QuicListenerThunk, this, &QuicListener))) 
+	{
+		UE_LOGFMT(LogSwitchboard, Error, "MsQuic ListenerOpen failed with status {Status}", static_cast<int64>(Status));
+		return false;
+	}
+
+	// Starts listening for incoming connections.
+	QUIC_ADDR QuicAddr = QuicAddrFromEndpoint(*ListenerEndpoint);
+	if (QUIC_FAILED(Status = QuicApi->ListenerStart(QuicListener, &SblAlpn, 1, &QuicAddr))) 
+	{
+		UE_LOGFMT(LogSwitchboard, Error, "MsQuic ListenerStart failed with status {Status}", static_cast<int64>(Status));
+		return false;
+	}
+
+	return true;
+}
+
+bool FSwitchboardListener::StopListening()
+{
+	if (QuicApi)
+	{
+		if (QuicListener)
+		{
+			QuicApi->ListenerClose(QuicListener);
+			QuicListener = nullptr;
+		}
+
+		if (QuicConfiguration)
+		{
+			QuicApi->ConfigurationClose(QuicConfiguration);
+			QuicConfiguration = nullptr;
+		}
+
+		if (QuicRegistration)
+		{
+			// This will block until all outstanding child objects have been closed.
+			QuicApi->RegistrationClose(QuicRegistration);
+			QuicRegistration = nullptr;
+		}
+
+		MsQuicClose(QuicApi);
+		QuicApi = nullptr;
+	}
+
+	return true;
+}
+
+void FSwitchboardListener::Tick()
+{
+	SWITCHBOARD_TRACE_CPUPROFILER_EVENT_SCOPE(FSwitchboardListener::Tick);
+
+	// Parse incoming data from remote connections
+	{
+		UE::TUniqueLock<UE::FRecursiveMutex> ConnectionsScopeLock(ConnectionsLock);
+		for (const TPair<FIPv4Endpoint, FConnectionRef>& ConnectionPair : ConnectionsByEndpoint)
+		{
+			const FIPv4Endpoint& ClientEndpoint = ConnectionPair.Key;
+			FConnectionRef Connection = ConnectionPair.Value;
+			if (Connection->bMessageComplete && Connection->ReceiveLock.TryLock())
+			{
+				ON_SCOPE_EXIT{ Connection->ReceiveLock.Unlock(); };
+
+				const FString Message(UTF8_TO_TCHAR(Connection->ReceiveBuffer->GetData()));
+				ParseIncomingMessage(Message, ClientEndpoint, Connection);
+				Connection->ReceiveBuffer->Empty();
+				Connection->bMessageComplete = false;
+
+				// Resume receive events after processing the pending message.
+				QuicApi->StreamReceiveSetEnabled(Connection->QuicStream, true);
+			}
+		}
+	}
+
+	// Run the next queued task
+	if (!ScheduledTasks.IsEmpty())
+	{
+		TUniquePtr<FSwitchboardTask> Task;
+		ScheduledTasks.Dequeue(Task);
+		RunScheduledTask(*Task);
+	}
+
+	CleanUpDisconnectedSockets();
+	HandleRunningProcesses(RunningProcesses, true);
+	HandleRunningProcesses(FlipModeMonitors, false);
+	SendMessageFutures();
+	SBLHelper->Tick();
+
+	OnTickDelegate.Broadcast();
+}
+
+bool FSwitchboardListener::ParseIncomingMessage(const FString& InMessage, const FIPv4Endpoint& InEndpoint, const FConnectionRef& Connection)
+{
+	SWITCHBOARD_TRACE_CPUPROFILER_EVENT_SCOPE(FSwitchboardListener::ParseIncomingMessage);
+
+	FCreateTaskResult CreateTaskResult =
+		CreateTaskFromCommand(InMessage, InEndpoint, Connection->bAuthenticated);
+
+	if (CreateTaskResult.Status == ECreateTaskStatus::Success)
+	{
+		if (CreateTaskResult.Task->Type == ESwitchboardTaskType::Disconnect)
+		{
+			DisconnectTasks.Enqueue(MoveTemp(CreateTaskResult.Task));
+		}
+		else if (CreateTaskResult.Task->Type == ESwitchboardTaskType::KeepAlive)
+		{
+			UE::TUniqueLock<UE::FRecursiveMutex> ConnectionsScopeLock(ConnectionsLock);
+			ConnectionsByEndpoint[InEndpoint]->LastActivityTime = FPlatformTime::Seconds();
+		}
+		else
+		{
+			if (CreateTaskResult.bEcho)
+			{
+				UE_LOG(LogSwitchboard, Display, TEXT("Received %s command"), CreateTaskResult.Task->GetCommandName());
+			}
+
+			SendMessage(CreateCommandAcceptedMessage(CreateTaskResult.Task->TaskID), InEndpoint);
+			ScheduledTasks.Enqueue(MoveTemp(CreateTaskResult.Task));
+		}
+		return true;
+	}
+	else if (CreateTaskResult.Status == ECreateTaskStatus::Error_Unauthenticated)
+	{
+		static FString UnauthenticatedError = FString::Printf(TEXT("Rejecting '%s' command for unauthenticated client %s"),
+			*CreateTaskResult.CommandName.Get(TEXT("(?)")), *InEndpoint.ToString());
+
+		UE_LOG(LogSwitchboard, Error, TEXT("%s"), *UnauthenticatedError);
+
+		// just use an empty ID if we couldn't find one
+		SendMessage(CreateCommandDeclinedMessage(CreateTaskResult.MessageID.Get(FGuid()), UnauthenticatedError), InEndpoint);
+	}
+	else
+	{
+		if (!CreateTaskResult.MessageID.IsSet())
+		{
+			FGuid RecoveredID;
+			if (TryFindIdInBrokenMessage(InMessage, RecoveredID))
+			{
+				CreateTaskResult.MessageID = RecoveredID;
+			}
+		}
+
+		static FString ParseError = FString::Printf(TEXT("Could not parse message %s with ID %s"),
+			*InMessage,
+			CreateTaskResult.MessageID.IsSet() ? *CreateTaskResult.MessageID.GetValue().ToString() : TEXT("(?)"));
+		UE_LOG(LogSwitchboard, Error, TEXT("%s"), *ParseError);
+			
+		// just use an empty ID if we couldn't find one
+		SendMessage(CreateCommandDeclinedMessage(CreateTaskResult.MessageID.Get(FGuid()), ParseError), InEndpoint);
+	}
+
+	return false;
+}
+
+bool FSwitchboardListener::RunScheduledTask(const FSwitchboardTask& InTask)
+{
+	switch (InTask.Type)
+	{
+		case ESwitchboardTaskType::Authenticate:
+		{
+			const FSwitchboardAuthenticateTask& AuthTask = static_cast<const FSwitchboardAuthenticateTask&>(InTask);
+			return Task_Authenticate(AuthTask);
+		}
+		case ESwitchboardTaskType::Start:
+		{
+			const FSwitchboardStartTask& StartTask = static_cast<const FSwitchboardStartTask&>(InTask);
+			return Task_StartProcess(StartTask);
+		}
+		case ESwitchboardTaskType::Kill:
+		{
+			const FSwitchboardKillTask& KillTask = static_cast<const FSwitchboardKillTask&>(InTask);
+			return Task_KillProcess(KillTask);
+		}
+		case ESwitchboardTaskType::ReceiveFileFromClient:
+		{
+			const FSwitchboardReceiveFileFromClientTask& ReceiveFileFromClientTask = static_cast<const FSwitchboardReceiveFileFromClientTask&>(InTask);
+			return Task_ReceiveFileFromClient(ReceiveFileFromClientTask);
+		}
+		case ESwitchboardTaskType::RedeployListener:
+		{
+			const FSwitchboardRedeployListenerTask& RedeployListenerTask = static_cast<const FSwitchboardRedeployListenerTask&>(InTask);
+			const bool bRedeployOk = Task_RedeployListener(RedeployListenerTask);
+			if (!bRedeployOk)
+			{
+				RollbackRedeploy();
+			}
+			return bRedeployOk;
+		}
+		case ESwitchboardTaskType::SendFileToClient:
+		{
+			const FSwitchboardSendFileToClientTask& SendFileToClientTask = static_cast<const FSwitchboardSendFileToClientTask&>(InTask);
+			return Task_SendFileToClient(SendFileToClientTask);
+		}
+		case ESwitchboardTaskType::KeepAlive:
+		{
+			return true;
+		}
+		case ESwitchboardTaskType::GetSyncStatus:
+		{
+			const FSwitchboardGetSyncStatusTask& GetSyncStatusTask = static_cast<const FSwitchboardGetSyncStatusTask&>(InTask);
+			return Task_GetSyncStatus(GetSyncStatusTask);
+		}
+		case ESwitchboardTaskType::FixExeFlags:
+		{
+			const FSwitchboardFixExeFlagsTask& Task = static_cast<const FSwitchboardFixExeFlagsTask&>(InTask);
+			return Task_FixExeFlags(Task);
+		}
+		case ESwitchboardTaskType::RefreshMosaics:
+		{
+			const FSwitchboardRefreshMosaicsTask& Task = static_cast<const FSwitchboardRefreshMosaicsTask&>(InTask);
+			return Task_RefreshMosaics(Task);
+		}
+		case ESwitchboardTaskType::MinimizeWindows:
+		{
+			const FSwitchboardMinimizeWindowsTask& Task = static_cast<const FSwitchboardMinimizeWindowsTask&>(InTask);
+			return Task_MinimizeWindows(Task);
+		}
+		case ESwitchboardTaskType::SetInactiveTimeout:
+		{
+			const FSwitchboardSetInactiveTimeoutTask& Task = static_cast<const FSwitchboardSetInactiveTimeoutTask&>(InTask);
+			return Task_SetInactiveTimeout(Task);
+		}
+		case ESwitchboardTaskType::FreeListenerBinary:
+		{
+			const FSwitchboardFreeListenerBinaryTask& Task = static_cast<const FSwitchboardFreeListenerBinaryTask&>(InTask);
+			return Task_FreeListenerBinary(Task);
+		}
+		default:
+		{
+			static const FString Response = TEXT("Unknown Command detected");
+			CreateCommandDeclinedMessage(InTask.TaskID, Response);
+			return false;
+		}
+	}
+}
+
+static uint32 FindPidInFocus()
+{
+#if PLATFORM_WINDOWS
+	HWND WindowHandle = GetForegroundWindow();
+
+	if (WindowHandle)
+	{
+		DWORD PID = 0;
+		DWORD ThreadId = GetWindowThreadProcessId(WindowHandle, &PID);
+
+		if (ThreadId != 0)
+		{
+			return PID;
+		}
+	}
+#endif // PLATFORM_WINDOWS
+
+	return 0;
+}
+
+#if PLATFORM_WINDOWS
+static TArray<FString> RegistryGetValueNames(const HKEY Key)
+{
+	TArray<FString> Names;
+	const uint32 MaxLength = 1024;
+	TCHAR ValueName[MaxLength];
+	DWORD ValueLength = MaxLength;
+
+	while (!RegEnumValue(Key, Names.Num(), ValueName, &ValueLength, nullptr, nullptr, nullptr, nullptr))
+	{
+		Names.Add(ValueName);
+		ValueLength = MaxLength;
+	}
+
+	return Names;
+}
+#endif // PLATFORM_WINDOWS
+
+#if PLATFORM_WINDOWS
+static FString RegistryGetStringValueData(const HKEY Key, const FString& ValueName)
+{
+	const uint32 MaxLength = 4096;
+	TCHAR ValueData[MaxLength];
+	DWORD ValueLength = MaxLength;
+
+	if (RegQueryValueEx(Key, *ValueName, 0, 0, LPBYTE(ValueData), &ValueLength))
+	{
+		return TEXT("");
+	}
+
+	ValueData[MaxLength - 1] = '\0';
+
+	return FString(ValueData);
+}
+#endif // PLATFORM_WINDOWS
+
+#if PLATFORM_WINDOWS
+static bool RegistrySetStringValueData(const HKEY Key, const FString& ValueName, const FString& ValueData)
+{
+	return !!RegSetValueEx(Key, *ValueName, 0, REG_SZ, LPBYTE(*ValueData), sizeof(TCHAR) * ValueData.Len());
+}
+#endif // PLATFORM_WINDOWS
+
+static bool DisableFullscreenOptimizationForProcess(const FRunningProcess* Process)
+{
+#if !PLATFORM_WINDOWS
+	return false;
+#else
+	// No point in continuing if there is no process to set the flags for.
+	if (!Process)
+	{
+		return false;
+	}
+
+	// This is the absolute path of the program we'll be looking for in the registry
+	const FString ProcessAbsolutePath = IFileManager::Get().ConvertToAbsolutePathForExternalAppForRead(*Process->Path);
+	const FString DisableFsOptLayer = TEXT("DISABLEDXMAXIMIZEDWINDOWEDMODE");
+	const FString LayersKeyPath = TEXT("Software\\Microsoft\\Windows NT\\CurrentVersion\\AppCompatFlags\\Layers");
+
+	// Check if the key exists
+	HKEY LayersKey;
+	if (RegOpenKeyExW(HKEY_CURRENT_USER, *LayersKeyPath, 0, KEY_ALL_ACCESS, &LayersKey))
+	{
+		return false;
+	}
+
+	ON_SCOPE_EXIT
+	{
+		RegCloseKey(LayersKey);
+	};
+
+	// If the key exists, the Value Names are the paths to the programs
+	const TArray<FString> ProgramPaths = RegistryGetValueNames(LayersKey);
+	const FString* ExistingPath = Algo::FindBy(ProgramPaths, ProcessAbsolutePath, [](const FString& ProgPath) { return IFileManager::Get().ConvertToAbsolutePathForExternalAppForRead(*ProgPath); });
+	if (ExistingPath)
+	{
+		// If the program already has an entry, get the layers from the Value Data.
+		FString ProgramLayers = RegistryGetStringValueData(LayersKey, *ExistingPath);
+
+		TArray<FString> ProgramLayersArray;
+		ProgramLayers.ParseIntoArray(ProgramLayersArray, TEXT(" "));
+
+		if (ProgramLayersArray.Contains(DisableFsOptLayer))
+		{
+			// Layer already present, nothing to do.
+			return true;
+		}
+
+		// Append the desired layer.
+		ProgramLayers = FString::Printf(TEXT("%s %s"), *ProgramLayers, *DisableFsOptLayer);
+		return RegistrySetStringValueData(LayersKey, *ExistingPath, ProgramLayers);
+	}
+	else
+	{
+		// The path to our executable does not exist as one of the regkey values, so we need to create it.
+		const FString ProgramLayers = FString::Printf(TEXT("~ %s"), *DisableFsOptLayer);
+		const FString ProgramAbsPath = ProcessAbsolutePath.Replace(TEXT("/"), TEXT("\\"));
+		return RegistrySetStringValueData(LayersKey, ProgramAbsPath, ProgramLayers);
+	}
+#endif // PLATFORM_WINDOWS
+}
+
+bool FSwitchboardListener::Task_Authenticate(const FSwitchboardAuthenticateTask& InAuthTask)
+{
+	SWITCHBOARD_TRACE_CPUPROFILER_EVENT_SCOPE(FSwitchboardListener::Task_Authenticate);
+
+	UE::TUniqueLock<UE::FRecursiveMutex> ConnectionsScopeLock(ConnectionsLock);
+	FConnectionRef Connection = ConnectionsByEndpoint[InAuthTask.Recipient];
+
+	auto SendAuthResponse = [this, &InAuthTask, &Connection]() {
+		SendMessage(CreateMessage(
+			{
+				{ TEXT("command"), FSwitchboardAuthenticateTask::CommandName },
+				{ TEXT("id"), InAuthTask.TaskID.ToString() },
+				{ TEXT("bAuthenticated"), Connection->bAuthenticated ? TEXT("true") : TEXT("false") },
+			}),
+			InAuthTask.Recipient);
+		};
+
+	if (Connection->bAuthenticated)
+	{
+		UE_LOGFMT(LogSwitchboard, Warning, "Ignoring redundant auth attempt for authenticated client {Endpoint}", Connection->Endpoint.ToString());
+		SendAuthResponse();
+		return false;
+	}
+
+	int8& AuthFailures = AuthFailuresByAddress.FindOrAdd(InAuthTask.Recipient.Address);
+	if (AuthFailures >= MaxAuthFailures)
+	{
+		UE_LOGFMT(LogSwitchboard, Error, "Ignoring auth attempt for client with excessive failures {Endpoint}", Connection->Endpoint.ToString());
+		SendAuthResponse();
+		return false;
+	}
+
+	if (InAuthTask.Password == ExpectedAuthPassword)
+	{
+		UE_LOGFMT(LogSwitchboard, Display, "Client {Endpoint} authenticated successfully", Connection->Endpoint.ToString());
+		Connection->bAuthenticated = true;
+		// TODO: Issue JWT?
+		SendAuthResponse();
+		return true;
+	}
+
+	UE_LOGFMT(LogSwitchboard, Warning, "Failed authentication attempt for client {Endpoint}", Connection->Endpoint.ToString());
+	++AuthFailures;
+	SendAuthResponse();
+	return false;
+}
+
+bool FSwitchboardListener::Task_StartProcess(const FSwitchboardStartTask& InRunTask)
+{
+	SWITCHBOARD_TRACE_CPUPROFILER_EVENT_SCOPE(FSwitchboardListener::Task_StartProcess);
+
+	TSharedRef<FRunningProcess> NewProcess = MakeShared<FRunningProcess>();
+
+	NewProcess->Recipient = InRunTask.Recipient;
+	NewProcess->Path = InRunTask.Command;
+	NewProcess->Name = InRunTask.Name;
+	NewProcess->Caller = InRunTask.Caller;
+	NewProcess->bUpdateClientsWithStdout = InRunTask.bUpdateClientsWithStdout;
+	NewProcess->bLockGpuClock = InRunTask.bLockGpuClock;
+	NewProcess->UUID = InRunTask.TaskID; // Process ID is the same as the message ID.
+	NewProcess->PID = 0; // default value
+
+	if (!NewProcess->CreatePipes())
+	{
+		return false;
+	}
+	
+	const bool bLaunchDetached = false;
+	const bool bLaunchHidden = false;
+	const bool bLaunchReallyHidden = false;
+	const int32 PriorityModifier = InRunTask.PriorityModifier;
+	const TCHAR* WorkingDirectory = InRunTask.WorkingDir.IsEmpty() ? nullptr : *InRunTask.WorkingDir;
+
+	NewProcess->Handle = FPlatformProcess::CreateProc(
+		*InRunTask.Command, 
+		*InRunTask.Arguments, 
+		bLaunchDetached, 
+		bLaunchHidden, 
+		bLaunchReallyHidden, 
+		&NewProcess->PID, 
+		PriorityModifier, 
+		WorkingDirectory, 
+		NewProcess->StdoutChildWritePipe,
+		NewProcess->StdinChildReadPipe
+	);
+
+	if (!NewProcess->Handle.IsValid())
+	{
+		// Close process in case it just didn't run
+		FPlatformProcess::CloseProc(NewProcess->Handle);
+
+		// close pipes
+		NewProcess->ClosePipes();
+
+		// log error
+		const FString ErrorMsg = FString::Printf(TEXT("Could not start program %s"), *InRunTask.Command);
+		UE_LOG(LogSwitchboard, Error, TEXT("%s"), *ErrorMsg);
+
+		// notify Switchboard
+		SendMessage(
+			CreateTaskDeclinedMessage(
+				InRunTask, 
+				ErrorMsg, 
+				{
+					{ TEXT("puuid"), NewProcess->UUID.ToString() },
+				}
+			), 
+			InRunTask.Recipient
+		);
+
+		return false;
+	}
+
+	// Lock Gpu Clocks for the lifetime of this PID, if requested
+	if (InRunTask.bLockGpuClock && SBLHelper.IsValid())
+	{
+		// Try to connect to the SBLHelper server if we haven't already
+		if (!SBLHelper->IsConnected())
+		{
+			FSBLHelperClient::FConnectionParams ConnectionParams;
+
+			uint16 Port = 8010; // Default tcp port
+
+			// Apply command line port number override, if present
+			{
+				static uint16 CmdLinePortOverride = 0;
+				static bool bCmdLinePortOverrideParsed = false;
+				static bool bCmdLinePortOverrideValid = false;
+
+				if (!bCmdLinePortOverrideParsed)
+				{
+					bCmdLinePortOverrideParsed = true;
+
+					bCmdLinePortOverrideValid = FParse::Value(FCommandLine::Get(), TEXT("sblhport="), CmdLinePortOverride);
+				}
+
+				if (bCmdLinePortOverrideValid)
+				{
+					Port = CmdLinePortOverride;
+				}
+			}
+
+			const FString HostName = FString::Printf(TEXT("localhost:%d"), Port);
+			FIPv4Endpoint::FromHostAndPort(*HostName, ConnectionParams.Endpoint);
+
+			SBLHelper->Connect(ConnectionParams);
+		}
+
+		if (SBLHelper->IsConnected())
+		{
+			const bool bSentMessage = SBLHelper->LockGpuClock(NewProcess->PID);
+
+			if (!bSentMessage)
+			{
+				UE_LOG(LogSwitchboard, Error, TEXT("Failed to send message to SBLHelper server to request gpu clock locking"));
+			}
+
+			// We disconnect right away because launches happen only far and in between.
+			SBLHelper->Disconnect();
+		}
+	}
+
+	UE_LOG(LogSwitchboard, Display, TEXT("Started process %d: %s %s"), NewProcess->PID, *InRunTask.Command, *InRunTask.Arguments);
+
+	RunningProcesses.Add(NewProcess);
+
+	// send message
+	{
+		FSwitchboardProgramStarted Packet;
+
+		Packet.Process.Uuid = NewProcess->UUID.ToString();
+		Packet.Process.Name = NewProcess->Name;
+		Packet.Process.Path = NewProcess->Path;
+		Packet.Process.Caller = NewProcess->Caller;
+		Packet.Process.Pid = NewProcess->PID;
+
+		UE::TUniqueLock<UE::FRecursiveMutex> ConnectionsScopeLock(ConnectionsLock);
+		for (const TPair<FIPv4Endpoint, FConnectionRef>& Connection : ConnectionsByEndpoint)
+		{
+			const FIPv4Endpoint& ClientEndpoint = Connection.Key;
+			SendMessage(CreateMessage(Packet), ClientEndpoint);
+		}
+	}
+
+	return true;
+}
+
+bool FSwitchboardListener::Task_KillProcess(const FSwitchboardKillTask& KillTask)
+{
+	SWITCHBOARD_TRACE_CPUPROFILER_EVENT_SCOPE(FSwitchboardListener::Task_KillProcess);
+
+	if (EquivalentTaskFutureExists(KillTask.GetEquivalenceHash()))
+	{
+		SendMessage(
+			CreateTaskDeclinedMessage(
+				KillTask,
+				TEXT("Duplicate"),
+				{
+					{ TEXT("puuid"), KillTask.ProgramID.ToString() },
+				}
+			),
+			KillTask.Recipient
+		);
+
+		return false;
+	}
+
+	// Look in RunningProcesses
+	TSharedPtr<FRunningProcess> Process;
+	{
+		TSharedPtr<FRunningProcess>* ProcessPtr = RunningProcesses.FindByPredicate([&KillTask](const TSharedPtr<FRunningProcess>& InProcess)
+		{
+			check(InProcess.IsValid());
+			return !InProcess->bPendingKill && (InProcess->UUID == KillTask.ProgramID);
+		});
+
+		if (ProcessPtr)
+		{
+			Process = *ProcessPtr;
+			Process->bPendingKill = true;
+		}
+	}
+
+	// Look in FlipModeMonitors
+	TSharedPtr<FRunningProcess> FlipModeMonitor;
+	{
+		TSharedPtr<FRunningProcess>* FlipModeMonitorPtr = FlipModeMonitors.FindByPredicate([&](const TSharedPtr<FRunningProcess>& FlipMonitor)
+		{
+			check(FlipMonitor.IsValid());
+			return !FlipMonitor->bPendingKill && (FlipMonitor->UUID == KillTask.ProgramID);
+		});
+
+		if (FlipModeMonitorPtr)
+		{
+			FlipModeMonitor = *FlipModeMonitorPtr;
+			FlipModeMonitor->bPendingKill = true;
+		}
+	}
+
+	// Create our future message
+	FSwitchboardMessageFuture MessageFuture;
+
+	MessageFuture.InEndpoint = KillTask.Recipient;
+	MessageFuture.EquivalenceHash = KillTask.GetEquivalenceHash();
+
+	const FGuid UUID = KillTask.ProgramID;
+
+	MessageFuture.Future = Async(EAsyncExecution::ThreadPool, [this, KillTask, Process, FlipModeMonitor, UUID]() {
+		SWITCHBOARD_TRACE_CPUPROFILER_EVENT_SCOPE_STR("FSwitchboardListener::Task_KillProcess future closure");
+
+		const float SoftKillTimeout = 2.0f;
+
+		const bool bKilledProcess = KillProcessNow(Process.Get(), SoftKillTimeout);
+		KillProcessNow(FlipModeMonitor.Get(), SoftKillTimeout);
+
+		// Clear bPendingKill
+		if (Process)
+		{
+			Process->bPendingKill = false;
+		}
+
+		if (FlipModeMonitor)
+		{
+			FlipModeMonitor->bPendingKill = false;
+		}
+
+		// Return message
+		const FString ProgramID = UUID.ToString();
+
+		if (!bKilledProcess)
+		{
+			const FString KillError = FString::Printf(TEXT("Could not kill program with ID %s"), *ProgramID);
+
+			return CreateTaskDeclinedMessage(
+				KillTask,
+				KillError,
+				{
+					{ TEXT("puuid"), ProgramID },
+				}
+			);
+		}
+
+		FSwitchboardProgramKilled Packet;
+
+		if (Process) // must be valid if it was killed
+		{
+			Packet.Process.Uuid = Process->UUID.ToString();
+			Packet.Process.Name = Process->Name;
+			Packet.Process.Path = Process->Path;
+			Packet.Process.Caller = Process->Caller;
+			Packet.Process.Pid = Process->PID;
+		}
+
+		return CreateMessage(Packet);
+	});
+
+	// Queue it to be sent when ready
+	MessagesFutures.Emplace(MoveTemp(MessageFuture));
+
+	return true;
+}
+
+bool FSwitchboardListener::Task_FixExeFlags(const FSwitchboardFixExeFlagsTask& Task)
+{
+	SWITCHBOARD_TRACE_CPUPROFILER_EVENT_SCOPE(FSwitchboardListener::Task_FixExeFlags);
+
+	if (EquivalentTaskFutureExists(Task.GetEquivalenceHash()))
+	{
+		SendMessage(	
+			CreateTaskDeclinedMessage(
+				Task,
+				TEXT("Duplicate"),
+				{
+					{ TEXT("puuid"), Task.ProgramID.ToString() },
+				}
+			),
+			Task.Recipient
+		);
+
+		return false;
+	}
+
+	// Look in RunningProcesses
+	TSharedPtr<FRunningProcess> Process;
+	{
+		TSharedPtr<FRunningProcess>* ProcessPtr = RunningProcesses.FindByPredicate([&Task](const TSharedPtr<FRunningProcess>& InProcess)
+			{
+				check(InProcess.IsValid());
+				return !InProcess->bPendingKill && (InProcess->UUID == Task.ProgramID);
+			});
+
+		if (ProcessPtr)
+		{
+			Process = *ProcessPtr;
+		}
+	}
+
+	if (!Process.IsValid())
+	{
+		SendMessage(
+			CreateTaskDeclinedMessage(
+				Task,
+				TEXT("Could not find ProgramID"),
+				{
+					{ TEXT("puuid"), Task.ProgramID.ToString() },
+				}
+			),
+			Task.Recipient
+		);
+
+		return false;
+	}
+
+	return DisableFullscreenOptimizationForProcess(Process.Get());
+}
+
+
+bool FSwitchboardListener::KillProcessNow(FRunningProcess* InProcess, float SoftKillTimeout)
+{
+	SWITCHBOARD_TRACE_CPUPROFILER_EVENT_SCOPE(FSwitchboardListener::KillProcessNow);
+
+	if (InProcess && InProcess->Handle.IsValid() && FPlatformProcess::IsProcRunning(InProcess->Handle))
+	{
+		UE_LOG(LogSwitchboard, Display, TEXT("Killing app with PID %d"), InProcess->PID);
+
+#if PLATFORM_WINDOWS
+		// try a soft kill first
+		if(SoftKillTimeout > 0)
+		{
+			const FString Params = FString::Printf(TEXT("/PID %d"), InProcess->PID);
+
+			FString OutStdOut;
+
+			FPlatformProcess::ExecProcess(TEXT("TASKKILL"), *Params, nullptr, &OutStdOut, nullptr);
+
+			const double TimeoutTime = FPlatformTime::Seconds() + SoftKillTimeout;
+			const float SleepTime = 0.050f;
+
+			while(FPlatformTime::Seconds() < TimeoutTime && FPlatformProcess::IsProcRunning(InProcess->Handle))
+			{
+				FPlatformProcess::Sleep(SleepTime);
+			}
+		}
+#endif // PLATFORM_WINDOWS
+
+		if (FPlatformProcess::IsProcRunning(InProcess->Handle))
+		{
+			const bool bKillTree = true;
+			FPlatformProcess::TerminateProc(InProcess->Handle, bKillTree);
+		}
+
+		// Pipes will be closed in HandleRunningProcesses
+		return true;
+	}
+
+	return false;
+}
+
+bool FSwitchboardListener::Task_ReceiveFileFromClient(const FSwitchboardReceiveFileFromClientTask& InReceiveFileFromClientTask)
+{
+	SWITCHBOARD_TRACE_CPUPROFILER_EVENT_SCOPE(FSwitchboardListener::Task_ReceiveFileFromClient);
+
+	FString Destination = InReceiveFileFromClientTask.Destination;
+
+	if (Destination.Contains(TEXT("%TEMP%")))
+	{
+		FString TempDir = FPlatformProcess::UserTempDir();
+
+		if (TempDir.EndsWith(TEXT("/")) || TempDir.EndsWith(TEXT("\\")))
+		{
+			TempDir.LeftChopInline(1, false);
+		}
+
+		Destination.ReplaceInline(TEXT("%TEMP%"), *TempDir);
+	}
+	if (Destination.Contains(TEXT("%RANDOM%")))
+	{
+		const FString Path = FPaths::GetPath(Destination);
+		const FString Extension = FPaths::GetExtension(Destination, true);
+		Destination = FPaths::CreateTempFilename(*Path, TEXT(""), *Extension);
+	}
+	FPlatformMisc::NormalizePath(Destination);
+	FPaths::MakePlatformFilename(Destination);
+
+	if (FPaths::FileExists(Destination))
+	{
+		if (!InReceiveFileFromClientTask.bForceOverwrite)
+		{
+			const FString ErrorMsg = FString::Printf(TEXT("Destination %s already exist"), *Destination);
+			UE_LOG(LogSwitchboard, Error, TEXT("%s"), *ErrorMsg);
+			SendMessage(CreateReceiveFileFromClientFailedMessage(Destination, ErrorMsg), InReceiveFileFromClientTask.Recipient);
+			return false;
+		}
+		else
+		{
+			IPlatformFile& PlatformFile = FPlatformFileManager::Get().GetPlatformFile();
+			PlatformFile.SetReadOnly(*Destination, false);
+		}
+	}
+
+	TArray<uint8> DecodedFileContent = {};
+	FBase64::Decode(InReceiveFileFromClientTask.FileContent, DecodedFileContent);
+
+	UE_LOG(LogSwitchboard, Display, TEXT("Writing %d bytes to %s"), DecodedFileContent.Num(), *Destination);
+	if (FFileHelper::SaveArrayToFile(DecodedFileContent, *Destination))
+	{
+		SendMessage(CreateReceiveFileFromClientCompletedMessage(Destination), InReceiveFileFromClientTask.Recipient);
+		return true;
+	}
+
+	const FString ErrorMsg = FString::Printf(TEXT("Error while trying to write to %s"), *Destination);
+	UE_LOG(LogSwitchboard, Error, TEXT("%s"), *ErrorMsg);
+	SendMessage(CreateReceiveFileFromClientFailedMessage(Destination, ErrorMsg), InReceiveFileFromClientTask.Recipient);
+	return false;
+}
+
+bool FSwitchboardListener::Task_RedeployListener(const FSwitchboardRedeployListenerTask& InRedeployListenerTask)
+{
+	SWITCHBOARD_TRACE_CPUPROFILER_EVENT_SCOPE(FSwitchboardListener::Task_RedeployListener);
+
+	if (RedeployStatus.State != FRedeployStatus::EState::None)
+	{
+		const FString ErrorMsg(TEXT("Redeploy already in progress"));
+		UE_LOG(LogSwitchboard, Error, TEXT("Redeploy listener: %s"), *ErrorMsg);
+		SendMessage(CreateRedeployStatusMessage(InRedeployListenerTask.TaskID, false, ErrorMsg), InRedeployListenerTask.Recipient);
+		return false;
+	}
+
+	RedeployStatus.RequestingClient = InRedeployListenerTask.Recipient;
+	RedeployStatus.State = FRedeployStatus::EState::RequestReceived;
+
+	TArray<uint8> DecodedFileContent;
+	FBase64::Decode(InRedeployListenerTask.FileContent, DecodedFileContent);
+
+	FSHAHash Hash;
+	FSHA1::HashBuffer(DecodedFileContent.GetData(), DecodedFileContent.Num(), Hash.Hash);
+	if (Hash != InRedeployListenerTask.ExpectedHash)
+	{
+		const FString ErrorMsg = FString::Printf(TEXT("Incorrect checksum %s; expected %s"),
+			*Hash.ToString(), *InRedeployListenerTask.ExpectedHash.ToString());
+		UE_LOG(LogSwitchboard, Error, TEXT("Redeploy listener: %s"), *ErrorMsg);
+		SendMessage(CreateRedeployStatusMessage(InRedeployListenerTask.TaskID, false, ErrorMsg), InRedeployListenerTask.Recipient);
+		return false;
+	}
+
+	UE_LOG(LogSwitchboard, Display, TEXT("Successfully validated new listener checksum"));
+
+	const uint32 CurrentPid = FPlatformProcess::GetCurrentProcessId();
+
+	// NOTE: FPlatformProcess::ExecutablePath() is stale if we were moved while running.
+	const FString OriginalThisExePath = FPlatformProcess::GetApplicationName(CurrentPid);
+	const FString ThisExeDir = FPaths::GetPath(OriginalThisExePath);
+	const FString ThisExeExt = FPaths::GetExtension(OriginalThisExePath, true);
+
+	const FString TempNewListenerPath = FPaths::CreateTempFilename(*ThisExeDir, TEXT("SwitchboardListener-New"), *ThisExeExt);
+
+	UE_LOG(LogSwitchboard, Display, TEXT("Writing new listener to \"%s\" (%d bytes)"), *TempNewListenerPath, DecodedFileContent.Num());
+	if (!FFileHelper::SaveArrayToFile(DecodedFileContent, *TempNewListenerPath))
+	{
+		const FString ErrorMsg = FString::Printf(TEXT("Unable to write new listener to \"%s\""), *TempNewListenerPath);
+		UE_LOG(LogSwitchboard, Error, TEXT("Redeploy listener: %s"), *ErrorMsg);
+		SendMessage(CreateRedeployStatusMessage(InRedeployListenerTask.TaskID, false, ErrorMsg), InRedeployListenerTask.Recipient);
+		return false;
+	}
+
+	// Unbind the port so the redeployed child process can claim it.
+	StopListening();
+
+	// Create IPC semaphore.
+	const FString IpcSemaphoreName = GetIpcSemaphoreName(CurrentPid);
+	FPlatformProcess::FSemaphore* IpcSemaphore = FPlatformProcess::NewInterprocessSynchObject(IpcSemaphoreName, true);
+	ON_SCOPE_EXIT
+	{
+		if (IpcSemaphore)
+		{
+			FPlatformProcess::DeleteInterprocessSynchObject(IpcSemaphore);
+		}
+	};
+
+	if (IpcSemaphore == nullptr)
+	{
+		const FString ErrorMsg(TEXT("Unable to create IPC semaphore"));
+		UE_LOG(LogSwitchboard, Error, TEXT("Redeploy listener: %s"), *ErrorMsg);
+		SendMessage(CreateRedeployStatusMessage(InRedeployListenerTask.TaskID, false, ErrorMsg), InRedeployListenerTask.Recipient);
+		return false;
+	}
+
+	// We acquire to zero immediately; the child process will release/increment when initialized,
+	// and we can await that indication by trying to acquire a second time.
+	IpcSemaphore->Lock();
+
+	RedeployStatus.OriginalThisExePath = OriginalThisExePath;
+	RedeployStatus.TempNewListenerPath = TempNewListenerPath;
+	RedeployStatus.State = FRedeployStatus::EState::NewListenerWrittenTemp;
+
+	const FString RedeployArgs = Options.ToString() + FString::Printf(TEXT(" -RedeployFromPid=%u"), CurrentPid);
+	const bool bLaunchDetached = true;
+	const bool bLaunchHidden = false;
+	const bool bLaunchReallyHidden = false;
+	uint32* OutProcessId = nullptr;
+	const int32 PriorityModifier = 0;
+	const TCHAR* WorkingDirectory = nullptr;
+	void* PipeWriteChild = nullptr;
+	FProcHandle RedeployedProcHandle = FPlatformProcess::CreateProc(*TempNewListenerPath, *RedeployArgs, bLaunchDetached, bLaunchHidden, bLaunchReallyHidden, OutProcessId, PriorityModifier, WorkingDirectory, PipeWriteChild);
+	if (!RedeployedProcHandle.IsValid())
+	{
+		const FString ErrorMsg(TEXT("Unable to launch new listener"));
+		UE_LOG(LogSwitchboard, Error, TEXT("Redeploy listener: %s"), *ErrorMsg);
+		SendMessage(CreateRedeployStatusMessage(InRedeployListenerTask.TaskID, false, ErrorMsg), InRedeployListenerTask.Recipient);
+		return false;
+	}
+
+	RedeployStatus.ListenerProc = RedeployedProcHandle;
+	RedeployStatus.State = FRedeployStatus::EState::NewListenerStarted;
+
+	const FString TempThisListenerPath = FPaths::CreateTempFilename(*ThisExeDir, TEXT("SwitchboardListener-Old"), *ThisExeExt);
+	if (!IFileManager::Get().Move(*TempThisListenerPath, *OriginalThisExePath, false))
+	{
+		const uint32 LastError = FPlatformMisc::GetLastError();
+		const FString ErrorMsg = FString::Printf(TEXT("Unable to move old listener to \"%s\" (error code %u)"), *TempThisListenerPath, LastError);
+		UE_LOG(LogSwitchboard, Error, TEXT("Redeploy listener: %s"), *ErrorMsg);
+		SendMessage(CreateRedeployStatusMessage(InRedeployListenerTask.TaskID, false, ErrorMsg), InRedeployListenerTask.Recipient);
+		return false;
+	}
+
+	RedeployStatus.TempThisListenerPath = TempThisListenerPath;
+	RedeployStatus.State = FRedeployStatus::EState::ThisListenerRenamed;
+
+	if (!IFileManager::Get().Move(*OriginalThisExePath, *TempNewListenerPath, false))
+	{
+		const uint32 LastError = FPlatformMisc::GetLastError();
+		const FString ErrorMsg = FString::Printf(TEXT("Unable to move new listener to \"%s\" (error code %u)"), *TempThisListenerPath, LastError);
+		UE_LOG(LogSwitchboard, Error, TEXT("Redeploy listener: %s"), *ErrorMsg);
+		SendMessage(CreateRedeployStatusMessage(InRedeployListenerTask.TaskID, false, ErrorMsg), InRedeployListenerTask.Recipient);
+		return false;
+	}
+
+	RedeployStatus.State = FRedeployStatus::EState::NewListenerRenamed;
+
+	const uint64 IpcReadyWaitSeconds = 5;
+	const bool IpcAcquired = IpcSemaphore->TryLock(IpcReadyWaitSeconds * 1'000'000'000ULL);
+	if (IpcAcquired == false)
+	{
+		const FString ErrorMsg(TEXT("Timed out waiting for child process to sync"));
+		UE_LOG(LogSwitchboard, Error, TEXT("Redeploy listener: %s"), *ErrorMsg);
+		SendMessage(CreateRedeployStatusMessage(InRedeployListenerTask.TaskID, false, ErrorMsg), InRedeployListenerTask.Recipient);
+		return false;
+	}
+
+	UE_LOG(LogSwitchboard, Display, TEXT("Redeploy complete; shutting down"));
+	SendMessage(CreateRedeployStatusMessage(InRedeployListenerTask.TaskID, true, FString(TEXT("Redeployed instance is ready"))), InRedeployListenerTask.Recipient);
+
+	RedeployStatus.State = FRedeployStatus::EState::Complete;
+	RequestEngineExit(TEXT("Redeploy complete"));
+
+	return true;
+}
+
+void FSwitchboardListener::RollbackRedeploy()
+{
+	SWITCHBOARD_TRACE_CPUPROFILER_EVENT_SCOPE(FSwitchboardListener::RollbackRedeploy);
+
+	// Unwind redeploy actions, in reverse order of corresponding state transitions.
+	UE_LOG(LogSwitchboard, Warning, TEXT("Rolling back redeploy..."));
+
+	if ((uint8)RedeployStatus.State >= (uint8)FRedeployStatus::EState::NewListenerRenamed)
+	{
+		UE_LOG(LogSwitchboard, Warning, TEXT("Moving new listener back to temporary location"));
+		IFileManager::Get().Move(*RedeployStatus.TempNewListenerPath, *RedeployStatus.OriginalThisExePath, false);
+	}
+
+	if ((uint8)RedeployStatus.State >= (uint8)FRedeployStatus::EState::ThisListenerRenamed)
+	{
+		UE_LOG(LogSwitchboard, Warning, TEXT("Moving this listener back to original location"));
+		IFileManager::Get().Move(*RedeployStatus.OriginalThisExePath, *RedeployStatus.TempThisListenerPath, false);
+		RedeployStatus.TempThisListenerPath.Reset();
+	}
+
+	if ((uint8)RedeployStatus.State >= (uint8)FRedeployStatus::EState::NewListenerStarted)
+	{
+		UE_LOG(LogSwitchboard, Warning, TEXT("Terminating temporary listener"));
+		FPlatformProcess::TerminateProc(RedeployStatus.ListenerProc);
+		FPlatformProcess::WaitForProc(RedeployStatus.ListenerProc);
+		FPlatformProcess::CloseProc(RedeployStatus.ListenerProc);
+	}
+
+	if ((uint8)RedeployStatus.State >= (uint8)FRedeployStatus::EState::NewListenerWrittenTemp)
+	{
+		UE_LOG(LogSwitchboard, Warning, TEXT("Removing temporary listener"));
+		IFileManager::Get().Delete(*RedeployStatus.TempNewListenerPath, true);
+		RedeployStatus.OriginalThisExePath.Reset();
+		RedeployStatus.TempNewListenerPath.Reset();
+
+		// Resume listening on this instance.
+		StartListening();
+	}
+
+	if ((uint8)RedeployStatus.State >= (uint8)FRedeployStatus::EState::RequestReceived)
+	{
+		RedeployStatus.RequestingClient = FIPv4Endpoint::Any;
+		RedeployStatus.State = FRedeployStatus::EState::None;
+	}
+}
+
+//static
+FString FSwitchboardListener::GetIpcSemaphoreName(uint32 ParentPid)
+{
+	return FString::Printf(TEXT("SwitchboardListenerIpc_Redeploy_%u"), ParentPid);
+}
+
+bool FSwitchboardListener::Task_SendFileToClient(const FSwitchboardSendFileToClientTask& InSendFileToClientTask)
+{
+	SWITCHBOARD_TRACE_CPUPROFILER_EVENT_SCOPE(FSwitchboardListener::Task_SendFileToClient);
+
+	FString SourceFilePath = InSendFileToClientTask.Source;
+	FPlatformMisc::NormalizePath(SourceFilePath);
+	FPaths::MakePlatformFilename(SourceFilePath);
+
+	if (!FPaths::FileExists(SourceFilePath))
+	{
+		const FString ErrorMsg = FString::Printf(TEXT("Could not find file %s"), *SourceFilePath);
+		UE_LOG(LogSwitchboard, Error, TEXT("%s"), *ErrorMsg);
+		SendMessage(CreateSendFileToClientFailedMessage(InSendFileToClientTask.Source, ErrorMsg), InSendFileToClientTask.Recipient);
+		return false;
+	}
+
+	TArray<uint8> FileContent;
+	if (!FFileHelper::LoadFileToArray(FileContent, *SourceFilePath))
+	{
+		const FString ErrorMsg = FString::Printf(TEXT("Error reading from file %s"), *SourceFilePath);
+		UE_LOG(LogSwitchboard, Error, TEXT("%s"), *ErrorMsg);
+		SendMessage(CreateSendFileToClientFailedMessage(InSendFileToClientTask.Source, ErrorMsg), InSendFileToClientTask.Recipient);
+		return false;
+	}
+
+	const FString EncodedFileContent = FBase64::Encode(FileContent);
+	return SendMessage(CreateSendFileToClientCompletedMessage(InSendFileToClientTask.Source, EncodedFileContent), InSendFileToClientTask.Recipient);
+}
+
+
+//static
+_Function_class_(QUIC_LISTENER_CALLBACK)
+QUIC_STATUS QUIC_API FSwitchboardListener::QuicListenerThunk(HQUIC Listener, void* Context, QUIC_LISTENER_EVENT* Event)
+{
+	FSwitchboardListener* This = reinterpret_cast<FSwitchboardListener*>(Context);
+	return This->QuicListenerCallback(Listener, Event);
+}
+
+
+QUIC_STATUS FSwitchboardListener::QuicListenerCallback(HQUIC Listener, QUIC_LISTENER_EVENT* Event)
+{
+	switch (Event->Type) {
+		case QUIC_LISTENER_EVENT_NEW_CONNECTION:
+		{
+			// A new connection is being attempted by a client. For the handshake to
+			// proceed, the server must provide a configuration for QUIC to use. The
+			// app MUST set the callback handler before returning.
+			const FIPv4Endpoint Endpoint = EndpointFromQuicAddr(*Event->NEW_CONNECTION.Info->RemoteAddress);
+			UE_LOGFMT(LogSwitchboard, Verbose, "[quic][{Endpoint}] LISTENER_EVENT_NEW_CONNECTION", Endpoint.ToString());
+
+			FConnectionRef Connection = MakeShared<FConnection>();
+			Connection->Endpoint = Endpoint;
+			Connection->QuicConn = Event->NEW_CONNECTION.Connection;
+			Connection->InactiveTimeout = DefaultInactiveTimeoutSeconds;
+			Connection->LastActivityTime = FPlatformTime::Seconds();
+
+			{
+				UE::TUniqueLock<UE::FRecursiveMutex> ConnectionsScopeLock(ConnectionsLock);
+				ConnectionsByEndpoint.Add(Endpoint, Connection);
+				ConnectionsByQuicConn.Add(Connection->QuicConn, Connection);
+			}
+
+			QuicApi->SetCallbackHandler(Connection->QuicConn, reinterpret_cast<void*>(&QuicConnectionThunk), this);
+			QuicApi->ConnectionSetConfiguration(Connection->QuicConn, QuicConfiguration);
+
+			return QUIC_STATUS_SUCCESS;
+		}
+		default:
+			break;
+	}
+
+	return QUIC_STATUS_NOT_SUPPORTED;
+}
+
+
+//static
+QUIC_STATUS QUIC_API FSwitchboardListener::QuicConnectionThunk(HQUIC Connection, void* Context, QUIC_CONNECTION_EVENT* Event)
+{
+	FSwitchboardListener* This = reinterpret_cast<FSwitchboardListener*>(Context);
+	return This->QuicConnectionCallback(Connection, Event);
+}
+
+
+QUIC_STATUS FSwitchboardListener::QuicConnectionCallback(HQUIC QuicConn, QUIC_CONNECTION_EVENT* Event)
+{
+	QUIC_ADDR RemoteAddr;
+	uint32_t BufferSize = sizeof(RemoteAddr);
+	QUIC_ENSURE(QuicApi->GetParam(QuicConn, QUIC_PARAM_CONN_REMOTE_ADDRESS, &BufferSize, &RemoteAddr));
+	const FIPv4Endpoint RemoteEndpoint = EndpointFromQuicAddr(RemoteAddr);
+
+	switch (Event->Type)
+	{
+		case QUIC_CONNECTION_EVENT_PEER_CERTIFICATE_RECEIVED:
+			// TODO, someday: Client certificate auth. //return QUIC_STATUS_UNKNOWN_CERTIFICATE;
+			break;
+		case QUIC_CONNECTION_EVENT_CONNECTED:
+			// The handshake has completed for the connection.
+			UE_LOGFMT(LogSwitchboard, Display, "[quic][conn][{Endpoint}] Connected", *RemoteEndpoint.ToString());
+			QuicApi->ConnectionSendResumptionTicket(QuicConn, QUIC_SEND_RESUMPTION_FLAG_NONE, 0, NULL);
+			break;
+		case QUIC_CONNECTION_EVENT_SHUTDOWN_INITIATED_BY_TRANSPORT:
+			// The connection has been shut down by the transport. Generally, this
+			// is the expected way for the connection to shut down with this
+			// protocol, since we let idle timeout kill the connection.
+			if (Event->SHUTDOWN_INITIATED_BY_TRANSPORT.Status == QUIC_STATUS_CONNECTION_IDLE)
+			{
+				UE_LOGFMT(LogSwitchboard, Display, "[quic][conn][{Endpoint}] Successfully shut down on idle.",
+					*RemoteEndpoint.ToString());
+			}
+			else
+			{
+				UE_LOGFMT(LogSwitchboard, Display, "[quic][conn][{Endpoint}] Shut down by transport, {Status}",
+					*RemoteEndpoint.ToString(), static_cast<int64>(Event->SHUTDOWN_INITIATED_BY_TRANSPORT.Status));
+			}
+			break;
+		case QUIC_CONNECTION_EVENT_SHUTDOWN_INITIATED_BY_PEER:
+			// The connection was explicitly shut down by the peer.
+			UE_LOGFMT(LogSwitchboard, Display, "[quic][conn][{Endpoint}] Shut down by peer, {ErrorCode}",
+				*RemoteEndpoint.ToString(), static_cast<uint64>(Event->SHUTDOWN_INITIATED_BY_PEER.ErrorCode));
+			break;
+		case QUIC_CONNECTION_EVENT_SHUTDOWN_COMPLETE:
+		{
+			// The connection has completed the shutdown process and is ready to be
+			// safely cleaned up.
+			UE_LOGFMT(LogSwitchboard, Display, "[quic][conn][{Endpoint}] Shutdown complete", *RemoteEndpoint.ToString());
+
+			{
+				UE::TUniqueLock<UE::FRecursiveMutex> ConnectionsScopeLock(ConnectionsLock);
+				FConnectionRef Connection = ConnectionsByQuicConn.FindAndRemoveChecked(QuicConn);
+				ConnectionsByEndpoint.FindAndRemoveChecked(Connection->Endpoint);
+				if (Connection->QuicStream)
+				{
+					ConnectionsByQuicStream.FindAndRemoveChecked(Connection->QuicStream);
+				}
+			}
+
+			QuicApi->ConnectionClose(QuicConn);
+			break;
+		}
+		case QUIC_CONNECTION_EVENT_PEER_STREAM_STARTED:
+		{
+			// The peer has started/created a new stream. The app MUST set the
+			// callback handler before returning.
+			UE_LOGFMT(LogSwitchboard, Display, "[quic][conn][{Endpoint}] Peer stream started", *RemoteEndpoint.ToString());
+
+			UE::TUniqueLock<UE::FRecursiveMutex> ConnectionsScopeLock(ConnectionsLock);
+
+			FConnectionRef Connection = ConnectionsByQuicConn[QuicConn];
+			ensure(Connection->QuicStream == nullptr); // Single stream per connection
+			Connection->QuicStream = Event->PEER_STREAM_STARTED.Stream;
+
+			ConnectionsByQuicStream.Add(Connection->QuicStream, Connection);
+
+			QuicApi->SetCallbackHandler(Connection->QuicStream, reinterpret_cast<void*>(&QuicStreamThunk), this);
+
+			// Send current state upon connection
+			{
+				FSwitchboardStatePacket StatePacket;
+
+				FPlatformMisc::GetOSVersions(StatePacket.OsVersionLabel, StatePacket.OsVersionLabelSub);
+				StatePacket.OsVersionNumber = FPlatformMisc::GetOSVersion();
+
+				StatePacket.TotalPhysicalMemory = FPlatformMemory::GetConstants().TotalPhysical;
+				StatePacket.PlatformBinaryDirectory = FPlatformProcess::GetBinariesSubdirectory();
+
+				for (const TSharedPtr<FRunningProcess>& RunningProcess : RunningProcesses)
+				{
+					check(RunningProcess.IsValid());
+
+					FSwitchboardStateRunningProcess StateRunningProcess;
+
+					StateRunningProcess.Uuid = RunningProcess->UUID.ToString();
+					StateRunningProcess.Name = RunningProcess->Name;
+					StateRunningProcess.Path = RunningProcess->Path;
+					StateRunningProcess.Caller = RunningProcess->Caller;
+					StateRunningProcess.Pid = RunningProcess->PID;
+
+					StatePacket.RunningProcesses.Add(MoveTemp(StateRunningProcess));
+				}
+
+				SendMessage(CreateMessage(StatePacket), RemoteEndpoint);
+			}
+
+			break;
+		}
+		case QUIC_CONNECTION_EVENT_RESUMED:
+			// The connection succeeded in doing a TLS resumption of a previous
+			// connection's session.
+			UE_LOGFMT(LogSwitchboard, Display, "[quic][conn][{Endpoint}] Connection resumed", *RemoteEndpoint.ToString());
+			break;
+		default:
+			break;
+	}
+
+	return QUIC_STATUS_SUCCESS;
+}
+
+
+//static
+QUIC_STATUS QUIC_API FSwitchboardListener::QuicStreamThunk(HQUIC Stream, void* Context, QUIC_STREAM_EVENT* Event)
+{
+	FSwitchboardListener* This = reinterpret_cast<FSwitchboardListener*>(Context);
+	return This->QuicStreamCallback(Stream, Event);
+}
+
+
+QUIC_STATUS FSwitchboardListener::QuicStreamCallback(HQUIC Stream, QUIC_STREAM_EVENT* Event)
+{
+	UE::TUniqueLock<UE::FRecursiveMutex> ConnectionsScopeLock(ConnectionsLock);
+	FConnectionRef Connection = ConnectionsByQuicStream[Stream];
+
+	QUIC_ADDR RemoteAddr;
+	uint32_t BufferSize = sizeof(RemoteAddr);
+	QUIC_ENSURE(QuicApi->GetParam(Connection->QuicConn, QUIC_PARAM_CONN_REMOTE_ADDRESS, &BufferSize, &RemoteAddr));
+	const FIPv4Endpoint RemoteEndpoint = EndpointFromQuicAddr(RemoteAddr);
+
+	QUIC_UINT62 StreamId;
+	BufferSize = sizeof(StreamId);
+	QUIC_ENSURE(QuicApi->GetParam(Stream, QUIC_PARAM_STREAM_ID, &BufferSize, &StreamId));
+
+	const FString StreamStr = FString::Printf(TEXT("%s[%llu]"), *RemoteEndpoint.ToString(), StreamId);
+
+	switch (Event->Type)
+	{
+		case QUIC_STREAM_EVENT_SEND_COMPLETE:
+		{
+			// A previous StreamSend call has completed, and the context is being
+			// returned back to the app.
+			UE::TUniqueLock<UE::FMutex> SendLock(Connection->SendLock);
+			TSharedPtr<FQuicBuffer> DequeuedSend;
+			Connection->SendBuffers.Dequeue(DequeuedSend);
+
+			// Single stream, ordered delivery, dequeue should release the right thing
+			ensure(DequeuedSend.Get() == Event->SEND_COMPLETE.ClientContext);
+
+			break;
+		}
+		case QUIC_STREAM_EVENT_RECEIVE:
+		{
+			// The value of this struct member when the callback returns is taken to mean
+			// how many bytes we actually consumed. Bytes left unconsumed will be handled
+			// after the tick re-enables receiving, upon consuming the preceding message.
+			Event->RECEIVE.TotalBufferLength = 0;
+
+			// Data was received from the peer on the stream.
+			Connection->LastActivityTime = FPlatformTime::Seconds();
+
+			UE::TUniqueLock<UE::FMutex> ReceiveLock(Connection->ReceiveLock);
+
+			if (!Connection->bMessageComplete)
+			{
+				for (uint32_t BufferNum = 0; BufferNum < Event->RECEIVE.BufferCount; ++BufferNum)
+				{
+					const QUIC_BUFFER Buffer = Event->RECEIVE.Buffers[BufferNum];
+					for (uint32 ReadIdx = 0; ReadIdx < Buffer.Length; ++ReadIdx)
+					{
+						const uint8_t Byte = Buffer.Buffer[ReadIdx];
+						Connection->ReceiveBuffer->Add(Byte);
+						++Event->RECEIVE.TotalBufferLength;
+
+						// If this concluded the outstanding message, signal the main thread to
+						// process it, and let MsQuic keep the next message buffered internally.
+						// Partial buffer consumption implicitly suspends receive events, which
+						// we then resume after the main thread processes the previous message.
+						if (Byte == 0)
+						{
+							Connection->bMessageComplete = true;
+
+							// Break inner loop + outer loop + switch
+							goto stream_receive_outer_break;
+						}
+					}
+				}
+			}
+
+stream_receive_outer_break:
+			break;
+		}
+		case QUIC_STREAM_EVENT_PEER_SEND_SHUTDOWN:
+			// The peer gracefully shut down its send direction of the stream.
+			UE_LOGFMT(LogSwitchboard, Display, "[quic][strm][{Stream}] Peer shut down", StreamStr);
+			break;
+		case QUIC_STREAM_EVENT_PEER_SEND_ABORTED:
+			// The peer aborted its send direction of the stream.
+			UE_LOGFMT(LogSwitchboard, Display, "[quic][strm][{Stream}] Peer aborted", StreamStr);
+			QuicApi->StreamShutdown(Stream, QUIC_STREAM_SHUTDOWN_FLAG_ABORT, 0);
+			break;
+		case QUIC_STREAM_EVENT_SHUTDOWN_COMPLETE:
+		{
+			// Both directions of the stream have been shut down and MsQuic is done
+			// with the stream. It can now be safely cleaned up.
+			UE_LOGFMT(LogSwitchboard, Display, "[quic][strm][{Stream}] Shutdown complete", StreamStr);
+
+			ensure(Connection->QuicStream == Stream);
+			ConnectionsByQuicStream.FindAndRemoveChecked(Stream);
+			QuicApi->StreamClose(Stream);
+			Connection->QuicStream = nullptr;
+
+			break;
+		}
+		default:
+			break;
+	}
+
+	return QUIC_STATUS_SUCCESS;
+}
+
+
+#if PLATFORM_WINDOWS
+static FCriticalSection SwitchboardListenerMutexNvapi;
+#endif // PLATFORM_WINDOWS
+
+#if PLATFORM_WINDOWS
+static void FillOutSyncTopologies(TArray<FSyncTopo>& SyncTopos)
+{
+	SWITCHBOARD_TRACE_CPUPROFILER_EVENT_SCOPE(FillOutSyncTopologies);
+
+	FScopeLock LockNvapi(&SwitchboardListenerMutexNvapi);
+
+	// Normally there is a single sync card. BUT an RTX Server could have more, and we need to account for that.
+
+	// Detect sync cards
+
+	NvU32 GSyncCount = 0;
+	NvGSyncDeviceHandle GSyncHandles[NVAPI_MAX_GSYNC_DEVICES];
+	NvAPI_GSync_EnumSyncDevices(GSyncHandles, &GSyncCount); // GSyncCount will be zero if error, so no need to check error.
+
+	for (NvU32 GSyncIdx = 0; GSyncIdx < GSyncCount; GSyncIdx++)
+	{
+		NvU32 GSyncGPUCount = 0;
+		NvU32 GSyncDisplayCount = 0;
+
+		// gather info first with null data pointers, just to get the count and subsequently allocate necessary memory.
+		{
+			const NvAPI_Status Result = NvAPI_GSync_GetTopology(GSyncHandles[GSyncIdx], &GSyncGPUCount, nullptr, &GSyncDisplayCount, nullptr);
+
+			if (Result != NVAPI_OK)
+			{
+				NvAPI_ShortString ErrorString;
+				NvAPI_GetErrorMessage(Result, ErrorString);
+				UE_LOG(LogSwitchboard, Warning, TEXT("NvAPI_GSync_GetTopology failed. Error: %s"), ANSI_TO_TCHAR(ErrorString));
+				continue;
+			}
+		}
+
+		// allocate memory for data
+		TArray<NV_GSYNC_GPU> GSyncGPUs;
+		TArray<NV_GSYNC_DISPLAY> GSyncDisplays;
+		{
+			GSyncGPUs.SetNumUninitialized(GSyncGPUCount, false);
+
+			for (NvU32 GSyncGPUIdx = 0; GSyncGPUIdx < GSyncGPUCount; GSyncGPUIdx++)
+			{
+				GSyncGPUs[GSyncGPUIdx].version = NV_GSYNC_GPU_VER;
+			}
+
+			GSyncDisplays.SetNumUninitialized(GSyncDisplayCount, false);
+
+			for (NvU32 GSyncDisplayIdx = 0; GSyncDisplayIdx < GSyncDisplayCount; GSyncDisplayIdx++)
+			{
+				GSyncDisplays[GSyncDisplayIdx].version = NV_GSYNC_DISPLAY_VER;
+			}
+		}
+
+		// get real info
+		{
+			const NvAPI_Status Result = NvAPI_GSync_GetTopology(GSyncHandles[GSyncIdx], &GSyncGPUCount, GSyncGPUs.GetData(), &GSyncDisplayCount, GSyncDisplays.GetData());
+
+			if (Result != NVAPI_OK)
+			{
+				NvAPI_ShortString ErrorString;
+				NvAPI_GetErrorMessage(Result, ErrorString);
+				UE_LOG(LogSwitchboard, Warning, TEXT("NvAPI_GSync_GetTopology failed. Error: %s"), ANSI_TO_TCHAR(ErrorString));
+				continue;
+			}
+		}
+
+		// Build outbound structure
+
+		FSyncTopo SyncTopo;
+
+		for (NvU32 GpuIdx = 0; GpuIdx < GSyncGPUCount; GpuIdx++)
+		{
+			FSyncGpu SyncGpu;
+
+			SyncGpu.bIsSynced = GSyncGPUs[GpuIdx].isSynced;
+			SyncGpu.Connector = int32(GSyncGPUs[GpuIdx].connector);
+
+			SyncTopo.SyncGpus.Emplace(SyncGpu);
+		}
+
+		for (NvU32 DisplayIdx = 0; DisplayIdx < GSyncDisplayCount; DisplayIdx++)
+		{
+			FSyncDisplay SyncDisplay;
+
+			switch (GSyncDisplays[DisplayIdx].syncState)
+			{
+			case NVAPI_GSYNC_DISPLAY_SYNC_STATE_UNSYNCED:
+				SyncDisplay.SyncState = TEXT("Unsynced");
+				break;
+			case NVAPI_GSYNC_DISPLAY_SYNC_STATE_SLAVE:
+				SyncDisplay.SyncState = TEXT("Follower");
+				break;
+			case NVAPI_GSYNC_DISPLAY_SYNC_STATE_MASTER:
+				SyncDisplay.SyncState = TEXT("Leader");
+				break;
+			default:
+				SyncDisplay.SyncState = TEXT("Unknown");
+				break;
+			}
+
+			// get color information for each display
+			{
+				NV_COLOR_DATA ColorData;
+
+				ColorData.version = NV_COLOR_DATA_VER;
+				ColorData.cmd = NV_COLOR_CMD_GET;
+				ColorData.size = sizeof(NV_COLOR_DATA);
+
+				const NvAPI_Status Result = NvAPI_Disp_ColorControl(GSyncDisplays[DisplayIdx].displayId, &ColorData);
+
+				if (Result == NVAPI_OK)
+				{
+					SyncDisplay.Bpc = ColorData.data.bpc;
+					SyncDisplay.Depth = ColorData.data.depth;
+					SyncDisplay.ColorFormat = ColorData.data.colorFormat;
+				}
+			}
+
+			SyncTopo.SyncDisplays.Emplace(SyncDisplay);
+		}
+
+		// Sync Status Parameters
+		{
+			NV_GSYNC_STATUS_PARAMS GSyncStatusParams;
+			GSyncStatusParams.version = NV_GSYNC_STATUS_PARAMS_VER;
+
+			const NvAPI_Status Result = NvAPI_GSync_GetStatusParameters(GSyncHandles[GSyncIdx], &GSyncStatusParams);
+
+			if (Result != NVAPI_OK)
+			{
+				NvAPI_ShortString ErrorString;
+				NvAPI_GetErrorMessage(Result, ErrorString);
+				UE_LOG(LogSwitchboard, Warning, TEXT("NvAPI_GSync_GetStatusParameters failed. Error: %s"), ANSI_TO_TCHAR(ErrorString));
+				continue;
+			}
+
+			SyncTopo.SyncStatusParams.RefreshRate = GSyncStatusParams.refreshRate;
+			SyncTopo.SyncStatusParams.HouseSyncIncoming = GSyncStatusParams.houseSyncIncoming;
+			SyncTopo.SyncStatusParams.bHouseSync = !!GSyncStatusParams.bHouseSync;
+			SyncTopo.SyncStatusParams.bInternalSecondary = GSyncStatusParams.bInternalSlave;
+		}
+
+		// Sync Control Parameters
+		{
+			NV_GSYNC_CONTROL_PARAMS GSyncControlParams;
+			GSyncControlParams.version = NV_GSYNC_CONTROL_PARAMS_VER;
+
+			const NvAPI_Status Result = NvAPI_GSync_GetControlParameters(GSyncHandles[GSyncIdx], &GSyncControlParams);
+
+			if (Result != NVAPI_OK)
+			{
+				NvAPI_ShortString ErrorString;
+				NvAPI_GetErrorMessage(Result, ErrorString);
+				UE_LOG(LogSwitchboard, Warning, TEXT("NvAPI_GSync_GetControlParameters failed. Error: %s"), ANSI_TO_TCHAR(ErrorString));
+				continue;
+			}
+
+			SyncTopo.SyncControlParams.bInterlaced = !!GSyncControlParams.interlaceMode;
+			SyncTopo.SyncControlParams.bSyncSourceIsOutput = !!GSyncControlParams.syncSourceIsOutput;
+			SyncTopo.SyncControlParams.Interval = GSyncControlParams.interval;
+			SyncTopo.SyncControlParams.Polarity = GSyncControlParams.polarity;
+			SyncTopo.SyncControlParams.Source = GSyncControlParams.source;
+			SyncTopo.SyncControlParams.VMode = GSyncControlParams.vmode;
+
+			SyncTopo.SyncControlParams.SyncSkew.MaxLines = GSyncControlParams.syncSkew.maxLines;
+			SyncTopo.SyncControlParams.SyncSkew.MinPixels = GSyncControlParams.syncSkew.minPixels;
+			SyncTopo.SyncControlParams.SyncSkew.NumLines = GSyncControlParams.syncSkew.numLines;
+			SyncTopo.SyncControlParams.SyncSkew.NumPixels = GSyncControlParams.syncSkew.numPixels;
+
+			SyncTopo.SyncControlParams.StartupDelay.MaxLines = GSyncControlParams.startupDelay.maxLines;
+			SyncTopo.SyncControlParams.StartupDelay.MinPixels = GSyncControlParams.startupDelay.minPixels;
+			SyncTopo.SyncControlParams.StartupDelay.NumLines = GSyncControlParams.startupDelay.numLines;
+			SyncTopo.SyncControlParams.StartupDelay.NumPixels = GSyncControlParams.startupDelay.numPixels;
+		}
+
+		SyncTopos.Emplace(SyncTopo);
+	}
+}
+#endif // PLATFORM_WINDOWS
+
+#if PLATFORM_WINDOWS
+static void FillOutDriverVersion(FSyncStatus& SyncStatus)
+{
+	SWITCHBOARD_TRACE_CPUPROFILER_EVENT_SCOPE(FillOutDriverVersion);
+
+	NvU32 DriverVersion;
+	NvAPI_ShortString BuildBranchString;
+
+	const NvAPI_Status Result = NvAPI_SYS_GetDriverAndBranchVersion(&DriverVersion, BuildBranchString);
+
+	if (Result != NVAPI_OK)
+	{
+		NvAPI_ShortString ErrorString;
+		NvAPI_GetErrorMessage(Result, ErrorString);
+		UE_LOG(LogSwitchboard, Warning, TEXT("NvAPI_SYS_GetDriverAndBranchVersion failed. Error: %s"), ANSI_TO_TCHAR(ErrorString));
+		return;
+	}
+
+	SyncStatus.DriverVersion = DriverVersion;
+	SyncStatus.DriverBranch = UTF8_TO_TCHAR(BuildBranchString);
+}
+#endif // PLATFORM_WINDOWS
+
+#if PLATFORM_WINDOWS
+static void FillOutTaskbarAutoHide(FSyncStatus& SyncStatus)
+{
+	SWITCHBOARD_TRACE_CPUPROFILER_EVENT_SCOPE(FillOutTaskbarAutoHide);
+
+	APPBARDATA AppBarData;
+
+	AppBarData.cbSize = sizeof(APPBARDATA);
+	AppBarData.hWnd = nullptr;
+	
+	const UINT Result = UINT(SHAppBarMessage(ABM_GETSTATE, &AppBarData));
+	
+	if (Result == ABS_AUTOHIDE)
+	{
+		SyncStatus.Taskbar = TEXT("AutoHide");
+	}
+	else
+	{
+		SyncStatus.Taskbar = TEXT("OnTop");
+	}
+}
+#endif // PLATFORM_WINDOWS
+
+
+#if PLATFORM_WINDOWS
+static void FillOutMosaicTopologies(TArray<FMosaicTopo>& MosaicTopos)
+{
+	SWITCHBOARD_TRACE_CPUPROFILER_EVENT_SCOPE(FillOutMosaicTopologies);
+
+	FScopeLock LockNvapi(&SwitchboardListenerMutexNvapi);
+
+	NvU32 GridCount = 0;
+	TArray<NV_MOSAIC_GRID_TOPO> GridTopologies;
+
+	// count how many grids
+	{
+		const NvAPI_Status Result = NvAPI_Mosaic_EnumDisplayGrids(nullptr, &GridCount);
+
+		if (Result != NVAPI_OK)
+		{
+			NvAPI_ShortString ErrorString;
+			NvAPI_GetErrorMessage(Result, ErrorString);
+			UE_LOG(LogSwitchboard, Warning, TEXT("NvAPI_Mosaic_EnumDisplayGrids failed. Error: %s"), ANSI_TO_TCHAR(ErrorString));
+			return;
+		}
+	}
+
+	// get the grids
+	{
+		GridTopologies.SetNumUninitialized(GridCount, false);
+
+		for (NvU32 TopoIdx = 0; TopoIdx < GridCount; TopoIdx++)
+		{
+			GridTopologies[TopoIdx].version = NV_MOSAIC_GRID_TOPO_VER;
+		}
+
+		const NvAPI_Status Result = NvAPI_Mosaic_EnumDisplayGrids(GridTopologies.GetData(), &GridCount);
+
+		if (Result != NVAPI_OK)
+		{
+			NvAPI_ShortString ErrorString;
+			NvAPI_GetErrorMessage(Result, ErrorString);
+			UE_LOG(LogSwitchboard, Warning, TEXT("NvAPI_Mosaic_EnumDisplayGrids failed. Error: %s"), ANSI_TO_TCHAR(ErrorString));
+			return;
+		}
+
+		for (NvU32 TopoIdx = 0; TopoIdx < GridCount; TopoIdx++)
+		{
+			FMosaicTopo MosaicTopo;
+			NV_MOSAIC_GRID_TOPO& GridTopo = GridTopologies[TopoIdx];
+
+			MosaicTopo.Columns = GridTopo.columns;
+			MosaicTopo.Rows = GridTopo.rows;
+			MosaicTopo.DisplayCount = GridTopo.displayCount;
+
+			MosaicTopo.DisplaySettings.Bpp = GridTopo.displaySettings.bpp;
+			MosaicTopo.DisplaySettings.Freq = GridTopo.displaySettings.freq;
+			MosaicTopo.DisplaySettings.Height = GridTopo.displaySettings.height;
+			MosaicTopo.DisplaySettings.Width = GridTopo.displaySettings.width;
+
+			MosaicTopos.Emplace(MosaicTopo);
+		}
+	}
+}
+#endif // PLATFORM_WINDOWS
+
+#if PLATFORM_WINDOWS
+FRunningProcess* FSwitchboardListener::FindOrStartFlipModeMonitorForUUID(const FGuid& UUID)
+{
+	SWITCHBOARD_TRACE_CPUPROFILER_EVENT_SCOPE(FSwitchboardListener::FindOrStartFlipModeMonitorForUUID);
+
+	// See if the associated FlipModeMonitor is running
+	{
+		TSharedPtr<FRunningProcess>* FlipModeMonitorPtr = FlipModeMonitors.FindByPredicate([&](const TSharedPtr<FRunningProcess>& FlipMonitor)
+		{
+			check(FlipMonitor.IsValid());
+			return FlipMonitor->UUID == UUID;
+		});
+
+		if (FlipModeMonitorPtr)
+		{
+			return (*FlipModeMonitorPtr).Get();
+		}
+	}
+
+	// It wasn't in there, so let's find our target process
+
+	TSharedPtr<FRunningProcess> Process;
+	{
+		TSharedPtr<FRunningProcess>* ProcessPtr = RunningProcesses.FindByPredicate([&](const TSharedPtr<FRunningProcess>& InProcess)
+		{
+			check(InProcess.IsValid());
+			return InProcess->UUID == UUID;
+		});
+
+		if (ProcessPtr)
+		{
+			Process = *ProcessPtr;
+		}
+	}
+
+	// If the target process does not exist, no point in continuing
+	if (!Process.IsValid())
+	{
+		return nullptr;
+	}
+
+	// Ok, we need to create our monitor.
+	TSharedRef<FRunningProcess> MonitorProcess = MakeShared<FRunningProcess>();
+
+	if (!MonitorProcess->CreatePipes())
+	{
+		UE_LOG(LogSwitchboard, Error, TEXT("Could not create pipe to read MonitorProcess output!"));
+		return nullptr;
+	}
+
+	const bool bLaunchDetached = true;
+	const bool bLaunchHidden = false;
+	const bool bLaunchReallyHidden = false;
+	const int32 PriorityModifier = 0;
+	const TCHAR* WorkingDirectory = nullptr;
+
+	MonitorProcess->Path = FPaths::EngineDir() / TEXT("Binaries") / TEXT("ThirdParty") / TEXT("PresentMon") / TEXT("Win64") / TEXT("PresentMon-1.8.0-x64.exe");
+
+	FString Arguments = 
+		FString::Printf(TEXT("-session_name session_%d -output_stdout -dont_restart_as_admin -terminate_on_proc_exit -stop_existing_session -process_id %d"), 
+		Process->PID, Process->PID);
+
+	MonitorProcess->Handle = FPlatformProcess::CreateProc(
+		*MonitorProcess->Path,
+		*Arguments,
+		bLaunchDetached,
+		bLaunchHidden,
+		bLaunchReallyHidden,
+		&MonitorProcess->PID,
+		PriorityModifier,
+		WorkingDirectory,
+		MonitorProcess->StdoutChildWritePipe,
+		MonitorProcess->StdinChildReadPipe
+	);
+
+	if (!MonitorProcess->Handle.IsValid() || !FPlatformProcess::IsProcRunning(MonitorProcess->Handle))
+	{
+		// Close process in case it just didn't run
+		FPlatformProcess::CloseProc(MonitorProcess->Handle);
+
+		// Close unused pipes
+		MonitorProcess->ClosePipes();
+
+		// Log error
+		const FString ErrorMsg = FString::Printf(TEXT("Could not start FlipMode monitor  %s"), *MonitorProcess->Path);
+		UE_LOG(LogSwitchboard, Error, TEXT("%s"), *ErrorMsg);
+
+		return nullptr;
+	}
+
+	// Log success
+	UE_LOG(LogSwitchboard, Display, TEXT("Started FlipMode monitor %d: %s %s"), MonitorProcess->PID, *MonitorProcess->Path, *Arguments);
+
+	// The UUID corresponds to the program being monitored. This will be used when looking for the Monitor of a given process.
+	// The monitor auto-closes when monitored program closes.
+	MonitorProcess->UUID = Process->UUID;
+	MonitorProcess->bUpdateClientsWithStdout = false;
+	MonitorProcess->bLockGpuClock = false;
+	MonitorProcess->Recipient = InvalidEndpoint;
+	MonitorProcess->Name = TEXT("flipmode_monitor");
+
+	FlipModeMonitors.Add(MonitorProcess);
+
+	return &MonitorProcess.Get();
+}
+#endif // PLATFORM_WINDOWS
+
+#if PLATFORM_WINDOWS
+static void FillOutFlipMode(FSyncStatus& SyncStatus, FRunningProcess* FlipModeMonitor)
+{
+	SWITCHBOARD_TRACE_CPUPROFILER_EVENT_SCOPE(FillOutFlipMode);
+
+	// See if the flip monitor is still there.
+	if (!FlipModeMonitor || !FlipModeMonitor->Handle.IsValid())
+	{
+		SyncStatus.FlipModeHistory.Add(TEXT("n/a")); // this informs Switchboard that data is not valid
+		return;
+	}
+
+	// Get stdout.
+	FUTF8ToTCHAR StdoutConv(reinterpret_cast<const ANSICHAR*>(FlipModeMonitor->Output.GetData()), FlipModeMonitor->Output.Num());
+	const FString StdOut(StdoutConv.Length(), StdoutConv.Get());
+
+	// Clear out the StdOut array.
+	FlipModeMonitor->Output.Empty();
+
+	// Split into lines
+
+	TArray<FString> Lines;
+	StdOut.ParseIntoArrayLines(Lines, false);
+
+	// Interpret the output as follows:
+	//
+	// Application,ProcessID,SwapChainAddress,Runtime,
+	// SyncInterval,PresentFlags,AllowsTearing,
+	// TimeInSeconds,MsBetweenPresents,MsBetweenDisplayChange,Dropped,
+	// PresentMode,
+	// MsInPresentAPI,MsUntilRenderComplete,MsUntilDisplayed
+	//
+	// e.g.
+	//   "UnrealEditor.exe,23256,0x000002DFBFE603A0,DXGI,1,0  ,0,35.65057220000000,1.91800000000000,30735.98660000000018,0,Composed: Flip,1.91820000000000,32.67810000000000,30700.80250000000160"
+
+	TArray<FString> Fields;
+
+	for (const FString& Line : Lines)
+	{
+		Line.ParseIntoArray(Fields, TEXT(","), false);
+		
+		if (Fields.Num() != 15)
+		{
+			continue;
+		}
+
+		const int32 PresentMonIdx = 11;
+
+		SyncStatus.FlipModeHistory.Add(Fields[PresentMonIdx]); // The first one will be "PresentMode". This is ok. 
+	}
+}
+#endif // PLATFORM_WINDOWS
+
+
+#if PLATFORM_WINDOWS
+static void FillOutDisableFullscreenOptimizationForProcess(FSyncStatus& SyncStatus, const FRunningProcess* Process)
+{
+	SWITCHBOARD_TRACE_CPUPROFILER_EVENT_SCOPE(FillOutDisableFullscreenOptimizationForProcess);
+
+	// Reset output array just in case
+	SyncStatus.ProgramLayers.Reset();
+
+	// No point in continuing if there is no process to get the flags for.
+	if (!Process)
+	{
+		return;
+	}
+
+	// This is the absolute path of the program we'll be looking for in the registry
+	const FString ProcessAbsolutePath = IFileManager::Get().ConvertToAbsolutePathForExternalAppForRead(*Process->Path);
+
+	const FString LayersKeyPath = TEXT("Software\\Microsoft\\Windows NT\\CurrentVersion\\AppCompatFlags\\Layers");
+
+	// Check if the key exists
+	HKEY LayersKey;
+	if (RegOpenKeyExW(HKEY_CURRENT_USER, *LayersKeyPath, 0, KEY_READ, &LayersKey))
+	{
+		return;
+	}
+
+	// If the key exists, the Value Names are the paths to the programs
+	const TArray<FString> ProgramPaths = RegistryGetValueNames(LayersKey);
+	const FString* ExistingPath = Algo::FindBy(ProgramPaths, ProcessAbsolutePath, [](const FString& ProgPath) { return IFileManager::Get().ConvertToAbsolutePathForExternalAppForRead(*ProgPath); });
+	if (ExistingPath)
+	{
+		// If so, get the layers from the Value Data.
+		const FString ProgramLayers = RegistryGetStringValueData(LayersKey, *ExistingPath);
+		ProgramLayers.ParseIntoArray(SyncStatus.ProgramLayers, TEXT(" "));
+	}
+
+	RegCloseKey(LayersKey);
+}
+#endif // PLATFORM_WINDOWS
+
+#if PLATFORM_WINDOWS
+static void FillOutPhysicalGpuStats(FSyncStatus& SyncStatus)
+{
+	SWITCHBOARD_TRACE_CPUPROFILER_EVENT_SCOPE(FillOutPhysicalGpuStats);
+
+	// TODO: Can we somehow use the GPU engine "engtype_3D" perf counters for this instead?
+
+	FScopeLock LockNvapi(&SwitchboardListenerMutexNvapi);
+
+	TArray<NvPhysicalGpuHandle> PhysicalGpuHandles;
+	PhysicalGpuHandles.SetNumUninitialized(NVAPI_MAX_PHYSICAL_GPUS);
+	NvU32 PhysicalGpuCount;
+	NvAPI_Status NvResult = NvAPI_EnumPhysicalGPUs(PhysicalGpuHandles.GetData(), &PhysicalGpuCount);
+	if (NvResult != NVAPI_OK)
+	{
+		NvAPI_ShortString ErrorString;
+		NvAPI_GetErrorMessage(NvResult, ErrorString);
+		UE_LOG(LogSwitchboard, Warning, TEXT("NvAPI_EnumPhysicalGPUs failed. Error: %s"), ANSI_TO_TCHAR(ErrorString));
+		return;
+	}
+
+	PhysicalGpuHandles.SetNum(PhysicalGpuCount);
+
+	// Sort first by bus, then by bus slot, ascending. Consistent with task manager and others.
+	Algo::Sort(PhysicalGpuHandles, [](const NvPhysicalGpuHandle& Lhs, const NvPhysicalGpuHandle& Rhs) -> bool {
+		NvU32 LhsBusId, RhsBusId, LhsSlotId, RhsSlotId;
+		const NvAPI_Status LhsBusResult = NvAPI_GPU_GetBusId(Lhs, &LhsBusId);
+		const NvAPI_Status RhsBusResult = NvAPI_GPU_GetBusId(Rhs, &RhsBusId);
+		const NvAPI_Status LhsSlotResult = NvAPI_GPU_GetBusSlotId(Lhs, &LhsSlotId);
+		const NvAPI_Status RhsSlotResult = NvAPI_GPU_GetBusSlotId(Rhs, &RhsSlotId);
+
+		if (LhsBusResult != NVAPI_OK || RhsBusResult != NVAPI_OK)
+		{
+			NvAPI_ShortString LhsErrorString;
+			NvAPI_ShortString RhsErrorString;
+			NvAPI_GetErrorMessage(LhsBusResult, LhsErrorString);
+			NvAPI_GetErrorMessage(RhsBusResult, RhsErrorString);
+			UE_LOG(LogSwitchboard, Warning, TEXT("NvAPI_GPU_GetBusId failed. Errors: %s, %s"), ANSI_TO_TCHAR(LhsErrorString), ANSI_TO_TCHAR(RhsErrorString));
+			return false;
+		}
+
+		if (LhsSlotResult != NVAPI_OK || RhsSlotResult != NVAPI_OK)
+		{
+			NvAPI_ShortString LhsErrorString;
+			NvAPI_ShortString RhsErrorString;
+			NvAPI_GetErrorMessage(LhsSlotResult, LhsErrorString);
+			NvAPI_GetErrorMessage(RhsSlotResult, RhsErrorString);
+			UE_LOG(LogSwitchboard, Warning, TEXT("NvAPI_GPU_GetBusSlotId failed. Errors: %s, %s"), ANSI_TO_TCHAR(LhsErrorString), ANSI_TO_TCHAR(RhsErrorString));
+			return false;
+		}
+
+		if (LhsBusId != RhsBusId)
+		{
+			return LhsBusId < RhsBusId;
+		}
+
+		return LhsSlotId < RhsSlotId;
+	});
+
+	SyncStatus.GpuUtilization.SetNumUninitialized(PhysicalGpuCount);
+	SyncStatus.GpuCoreClocksKhz.SetNumUninitialized(PhysicalGpuCount);
+	SyncStatus.GpuTemperature.SetNumUninitialized(PhysicalGpuCount);
+
+	for (NvU32 PhysicalGpuIdx = 0; PhysicalGpuIdx < PhysicalGpuCount; ++PhysicalGpuIdx)
+	{
+		SyncStatus.GpuUtilization[PhysicalGpuIdx] = -1;
+		SyncStatus.GpuCoreClocksKhz[PhysicalGpuIdx] = -1;
+		SyncStatus.GpuTemperature[PhysicalGpuIdx] = MIN_int32;
+
+		const NvPhysicalGpuHandle& PhysicalGpu = PhysicalGpuHandles[PhysicalGpuIdx];
+
+		NV_GPU_DYNAMIC_PSTATES_INFO_EX PstatesInfo;
+		PstatesInfo.version = NV_GPU_DYNAMIC_PSTATES_INFO_EX_VER;
+		NvResult = NvAPI_GPU_GetDynamicPstatesInfoEx(PhysicalGpu, &PstatesInfo);
+		if (NvResult == NVAPI_OK)
+		{
+			// FIXME: NV_GPU_UTILIZATION_DOMAIN_ID enum is missing in our nvapi.h, but documented elsewhere.
+			//const int8 UtilizationDomain = NVAPI_GPU_UTILIZATION_DOMAIN_GPU;
+			const int8 UtilizationDomain = 0;
+			if (PstatesInfo.utilization[UtilizationDomain].bIsPresent)
+			{
+				SyncStatus.GpuUtilization[PhysicalGpuIdx] = PstatesInfo.utilization[UtilizationDomain].percentage;
+			}
+		}
+		else
+		{
+			UE_LOG(LogSwitchboard, Warning, TEXT("NvAPI_GPU_GetDynamicPstatesInfoEx failed. Error code: %d"), NvResult);
+		}
+
+		NV_GPU_CLOCK_FREQUENCIES ClockFreqs;
+		ClockFreqs.version = NV_GPU_CLOCK_FREQUENCIES_VER;
+		ClockFreqs.ClockType = NV_GPU_CLOCK_FREQUENCIES_CURRENT_FREQ;
+		NvResult = NvAPI_GPU_GetAllClockFrequencies(PhysicalGpu, &ClockFreqs);
+		if (NvResult == NVAPI_OK)
+		{
+			if (ClockFreqs.domain[NVAPI_GPU_PUBLIC_CLOCK_GRAPHICS].bIsPresent)
+			{
+				SyncStatus.GpuCoreClocksKhz[PhysicalGpuIdx] = ClockFreqs.domain[NVAPI_GPU_PUBLIC_CLOCK_GRAPHICS].frequency;
+			}
+		}
+		else
+		{
+			UE_LOG(LogSwitchboard, Warning, TEXT("NvAPI_GPU_GetAllClockFrequencies failed. Error code: %d"), NvResult);
+		}
+
+		NV_GPU_THERMAL_SETTINGS ThermalSettings;
+		ThermalSettings.version = NV_GPU_THERMAL_SETTINGS_VER;
+		NvResult = NvAPI_GPU_GetThermalSettings(PhysicalGpu, NVAPI_THERMAL_TARGET_ALL, &ThermalSettings);
+		if (NvResult == NVAPI_OK)
+		{
+			// Report max temp across all sensors for this GPU.
+			for (NvU32 SensorIdx = 0; SensorIdx < ThermalSettings.count; ++SensorIdx)
+			{
+				const NvS32 SensorTemp = ThermalSettings.sensor[SensorIdx].currentTemp;
+				if (SensorTemp > SyncStatus.GpuTemperature[PhysicalGpuIdx])
+				{
+					SyncStatus.GpuTemperature[PhysicalGpuIdx] = SensorTemp;
+				}
+			}
+		}
+		else
+		{
+			UE_LOG(LogSwitchboard, Warning, TEXT("NvAPI_GPU_GetThermalSettings failed. Error code: %d"), NvResult);
+		}
+	}
+}
+#endif
+
+bool FSwitchboardListener::EquivalentTaskFutureExists(uint32 TaskEquivalenceHash) const
+{
+	return !!MessagesFutures.FindByPredicate([=](const FSwitchboardMessageFuture& MessageFuture)
+	{
+		return MessageFuture.EquivalenceHash == TaskEquivalenceHash;
+	});
+}
+
+bool FSwitchboardListener::Task_GetSyncStatus(const FSwitchboardGetSyncStatusTask& InGetSyncStatusTask)
+{
+	SWITCHBOARD_TRACE_CPUPROFILER_EVENT_SCOPE(FSwitchboardListener::Task_GetSyncStatus);
+
+#if PLATFORM_WINDOWS
+	// Reject request if an equivalent one is already in our future
+	if (EquivalentTaskFutureExists(InGetSyncStatusTask.GetEquivalenceHash()))
+	{
+		SendMessage(
+			CreateTaskDeclinedMessage(InGetSyncStatusTask, TEXT("Duplicate"),{}),
+			InGetSyncStatusTask.Recipient
+		);
+
+		return false;
+	}
+
+	TSharedRef<FSyncStatus> SyncStatus = MakeShared<FSyncStatus>(); // Smart pointer to avoid potentially bigger copy to lambda below.
+
+	// We need to run these on this thread to avoid threading issues.
+	FillOutFlipMode(SyncStatus.Get(), FindOrStartFlipModeMonitorForUUID(InGetSyncStatusTask.ProgramID));
+
+	// Fill out fullscreen optimization setting
+	{
+		TSharedPtr<FRunningProcess>* ProcessPtr = RunningProcesses.FindByPredicate([&](const TSharedPtr<FRunningProcess>& InProcess)
+		{
+			check(InProcess.IsValid());
+			return !InProcess->bPendingKill && (InProcess->UUID == InGetSyncStatusTask.ProgramID);
+		});
+
+		FRunningProcess* Process = nullptr;
+
+		if (ProcessPtr)
+		{
+			Process = (*ProcessPtr).Get();
+		}
+
+		FillOutDisableFullscreenOptimizationForProcess(SyncStatus.Get(), Process);
+	}
+
+	// Create our future message
+	FSwitchboardMessageFuture MessageFuture;
+
+	MessageFuture.InEndpoint = InGetSyncStatusTask.Recipient;
+	MessageFuture.EquivalenceHash = InGetSyncStatusTask.GetEquivalenceHash();
+
+	MessageFuture.Future = Async(EAsyncExecution::ThreadPool,
+		[
+			SyncStatus,
+			IsNvAPIInitialized=bIsNvAPIInitialized,
+			CpuMonitor=CpuMonitor,
+			CachedMosaicToposLock=CachedMosaicToposLock,
+			CachedMosaicTopos=CachedMosaicTopos
+		]() {
+			SWITCHBOARD_TRACE_CPUPROFILER_EVENT_SCOPE_STR("FSwitchboardListener::Task_GetSyncStatus future closure");
+
+			FillOutTaskbarAutoHide(SyncStatus.Get());
+
+			if (IsNvAPIInitialized)
+			{
+				FillOutDriverVersion(SyncStatus.Get());
+				FillOutSyncTopologies(SyncStatus->SyncTopos);
+			}
+
+			{
+				FReadScopeLock Lock(*CachedMosaicToposLock);
+				SyncStatus->MosaicTopos = *CachedMosaicTopos;
+			}
+
+			SyncStatus->PidInFocus = FindPidInFocus();
+
+			if (CpuMonitor)
+			{
+				CpuMonitor->GetPerCoreUtilization(SyncStatus->CpuUtilization);
+			}
+
+			const FPlatformMemoryStats MemStats = FPlatformMemory::GetStats();
+			SyncStatus->AvailablePhysicalMemory = MemStats.AvailablePhysical;
+
+			if (IsNvAPIInitialized)
+			{
+				FillOutPhysicalGpuStats(SyncStatus.Get());
+			}
+
+			return CreateSyncStatusMessage(SyncStatus.Get());
+		}
+	);
+
+	// Queue it to be sent when ready
+	MessagesFutures.Emplace(MoveTemp(MessageFuture));
+
+	return true;
+#else
+	SendMessage(CreateTaskDeclinedMessage(InGetSyncStatusTask, "Platform not supported", {}), InGetSyncStatusTask.Recipient);
+	return false;
+#endif // PLATFORM_WINDOWS
+}
+
+bool FSwitchboardListener::Task_RefreshMosaics(const FSwitchboardRefreshMosaicsTask& InRefreshMosaicsTask)
+{
+	SWITCHBOARD_TRACE_CPUPROFILER_EVENT_SCOPE(FSwitchboardListener::Task_RefreshMosaics);
+
+#if PLATFORM_WINDOWS
+	if (!bIsNvAPIInitialized)
+	{
+		SendMessage(
+			CreateTaskDeclinedMessage(InRefreshMosaicsTask, "NvAPI not supported", {}),
+			InRefreshMosaicsTask.Recipient);
+		return false;
+	}
+
+	// Reject request if an equivalent one is already in our future
+	if (EquivalentTaskFutureExists(InRefreshMosaicsTask.GetEquivalenceHash()))
+	{
+		SendMessage(
+			CreateTaskDeclinedMessage(InRefreshMosaicsTask, TEXT("Duplicate"), {}),
+			InRefreshMosaicsTask.Recipient
+		);
+
+		return false;
+	}
+
+	// Create our future
+	FSwitchboardMessageFuture TaskFuture;
+	TaskFuture.EquivalenceHash = InRefreshMosaicsTask.GetEquivalenceHash();
+	TaskFuture.Future = Async(EAsyncExecution::ThreadPool,
+		[
+			CachedMosaicToposLock=CachedMosaicToposLock,
+			CachedMosaicTopos=CachedMosaicTopos
+		]()
+		{
+			SWITCHBOARD_TRACE_CPUPROFILER_EVENT_SCOPE_STR("FSwitchboardListener::Task_RefreshMosaics future closure");
+
+			FWriteScopeLock Lock(*CachedMosaicToposLock);
+			CachedMosaicTopos->Reset();
+			FillOutMosaicTopologies(*CachedMosaicTopos);
+			
+			return FString();
+		}
+	);
+
+	MessagesFutures.Emplace(MoveTemp(TaskFuture));
+
+	return true;
+#else
+	SendMessage(CreateTaskDeclinedMessage(InRefreshMosaicsTask, "Platform not supported", {}), InRefreshMosaicsTask.Recipient);
+	return false;
+#endif // PLATFORM_WINDOWS
+}
+
+bool FSwitchboardListener::Task_MinimizeWindows(const FSwitchboardMinimizeWindowsTask& InTask)
+{
+	SWITCHBOARD_TRACE_CPUPROFILER_EVENT_SCOPE(FSwitchboardListener::Task_MinimizeWindows);
+
+#if PLATFORM_WINDOWS
+	// Reject request if an equivalent one is already in our future
+	if (EquivalentTaskFutureExists(InTask.GetEquivalenceHash()))
+	{
+		SendMessage(
+			CreateTaskDeclinedMessage(InTask, TEXT("Duplicate"), {}),
+			InTask.Recipient
+		);
+
+		return false;
+	}
+
+	// Create our future
+	FSwitchboardMessageFuture TaskFuture;
+	TaskFuture.EquivalenceHash = InTask.GetEquivalenceHash();
+	TaskFuture.Future = Async(EAsyncExecution::ThreadPool,
+		[]()
+		{
+			SWITCHBOARD_TRACE_CPUPROFILER_EVENT_SCOPE_STR("FSwitchboardListener::Task_MinimizeWindows future closure");
+
+			if (const HWND HWnd = ::FindWindowA("Shell_TrayWnd", NULL))
+			{
+				::SendMessageA(HWnd, WM_COMMAND, (WPARAM)419, 0);
+			}
+
+			return FString();
+		}
+	);
+
+	MessagesFutures.Emplace(MoveTemp(TaskFuture));
+
+	return true;
+#else
+	SendMessage(CreateTaskDeclinedMessage(InTask, "Platform not supported", {}), InTask.Recipient);
+	return false;
+#endif // PLATFORM_WINDOWS
+}
+
+bool FSwitchboardListener::Task_SetInactiveTimeout(const FSwitchboardSetInactiveTimeoutTask& InTimeoutTask)
+{
+	const FIPv4Endpoint& Client = InTimeoutTask.Recipient;
+	const float RequestedTimeout = InTimeoutTask.TimeoutSeconds;
+
+	if (RequestedTimeout < DefaultInactiveTimeoutSeconds)
+	{
+		SendMessage(CreateTaskDeclinedMessage(InTimeoutTask, "Requested timeout too low", {}), Client);
+		return false;
+	}
+
+	float& ClientTimeout = ConnectionsByEndpoint[Client]->InactiveTimeout;
+	if (RequestedTimeout != ClientTimeout)
+	{
+		UE_LOG(LogSwitchboard, Display, TEXT("Changing client %s inactive timeout from %.0f to %.0f seconds"),
+			*Client.ToString(), ClientTimeout, RequestedTimeout);
+		ClientTimeout = RequestedTimeout;
+	}
+
+	return true;
+}
+
+bool FSwitchboardListener::Task_FreeListenerBinary(const FSwitchboardFreeListenerBinaryTask& InFreeListenerBinaryTask)
+{
+	SWITCHBOARD_TRACE_CPUPROFILER_EVENT_SCOPE(FSwitchboardListener::Task_FreeListenerBinary);
+
+	const uint32 CurrentPid = FPlatformProcess::GetCurrentProcessId();
+
+	// NOTE: FPlatformProcess::ExecutablePath() is stale if we were moved while running.
+	const FString OriginalThisExePath = FPlatformProcess::GetApplicationName(CurrentPid);
+	const FString ThisExeDir = FPaths::GetPath(OriginalThisExePath);
+	const FString ThisExeFilename = FPaths::GetCleanFilename(OriginalThisExePath);
+	
+	const FString MovedListenerPath = FPaths::EngineIntermediateDir() / TEXT("Switchboard") / TEXT("old_") + ThisExeFilename;
+
+	//
+	// Weird hackery to get around Windows locking down an executable file on disk whilst running.
+	// We move the (locked) file (which is allowed), and then make a copy (which is no longer linked to the running process)
+
+	if (!IFileManager::Get().Move(*MovedListenerPath, *OriginalThisExePath, true))
+	{
+		const uint32 LastError = FPlatformMisc::GetLastError();
+		const FString ErrorMsg = FString::Printf(TEXT("Unable to move listener exe to \"%s\" (error code %u)"), *MovedListenerPath, LastError);
+		UE_LOG(LogSwitchboard, Error, TEXT("Free listener binary: %s"), *ErrorMsg);
+		SendMessage(CreateTaskDeclinedMessage(InFreeListenerBinaryTask, ErrorMsg, {}), InFreeListenerBinaryTask.Recipient);
+		return false;
+	}
+
+	if (IFileManager::Get().Copy(*OriginalThisExePath, *MovedListenerPath, true, true) != ECopyResult::COPY_OK)
+	{
+		const uint32 LastError = FPlatformMisc::GetLastError();
+		const FString ErrorMsg = FString::Printf(TEXT("Unable to copy listener exe back to \"%s\" (error code %u)"), *OriginalThisExePath, LastError);
+		UE_LOG(LogSwitchboard, Error, TEXT("Free listener binary: %s"), *ErrorMsg);
+		SendMessage(CreateTaskDeclinedMessage(InFreeListenerBinaryTask, ErrorMsg, {}), InFreeListenerBinaryTask.Recipient);
+		return false;
+	}
+
+	return true;
+}
+
+void FSwitchboardListener::CleanUpDisconnectedSockets()
+{
+	SWITCHBOARD_TRACE_CPUPROFILER_EVENT_SCOPE(FSwitchboardListener::CleanUpDisconnectedSockets);
+
+	UE::TUniqueLock<UE::FRecursiveMutex> ConnectionsScopeLock(ConnectionsLock);
+
+	const double CurrentTime = FPlatformTime::Seconds();
+	for (const TPair<FIPv4Endpoint, FConnectionRef>& Connection : ConnectionsByEndpoint)
+	{
+		const FIPv4Endpoint& Client = Connection.Key;
+		const float ClientTimeout = Connection.Value->InactiveTimeout;
+		const double InactiveTime = CurrentTime - Connection.Value->LastActivityTime;
+		if (ClientTimeout > 0.0f && InactiveTime > ClientTimeout)
+		{
+			UE_LOG(LogSwitchboard, Warning, TEXT("Client %s has been inactive for more than %.1fs -- closing connection"), *Client.ToString(), ClientTimeout);
+			TUniquePtr<FSwitchboardDisconnectTask> DisconnectTask = MakeUnique<FSwitchboardDisconnectTask>(FGuid(), Client);
+			DisconnectTasks.Enqueue(MoveTemp(DisconnectTask));
+		}
+	}
+
+	while (!DisconnectTasks.IsEmpty())
+	{
+		TUniquePtr<FSwitchboardTask> Task;
+		DisconnectTasks.Dequeue(Task);
+		const FSwitchboardDisconnectTask& DisconnectTask = static_cast<const FSwitchboardDisconnectTask&>(*Task);
+		const FIPv4Endpoint& Client = DisconnectTask.Recipient;
+
+		if (RedeployStatus.InProgress() && Client == RedeployStatus.RequestingClient)
+		{
+			UE_LOG(LogSwitchboard, Warning, TEXT("Client %s disconnected before redeploy completed"), *Client.ToString());
+			RollbackRedeploy();
+		}
+
+		UE_LOG(LogSwitchboard, Display, TEXT("Client %s disconnecting"), *Client.ToString());
+
+		if (FConnectionRef* MaybeConnection = ConnectionsByEndpoint.Find(Client))
+		{
+			if ((*MaybeConnection)->QuicConn)
+			{
+				QuicApi->ConnectionShutdown((*MaybeConnection)->QuicConn, QUIC_CONNECTION_SHUTDOWN_FLAG_SILENT, 0);
+			}
+		}
+	}
+}
+
+void FSwitchboardListener::HandleStdout(const TSharedPtr<FRunningProcess>& Process)
+{
+	SWITCHBOARD_TRACE_CPUPROFILER_EVENT_SCOPE(FSwitchboardListener::HandleStdout);
+
+	TArray<uint8> Output;
+	if (FPlatformProcess::ReadPipeToArray(Process->StdoutParentReadPipe, Output))
+	{
+		Process->Output.Append(Output);
+	}
+
+	// If there was a new stdout, update the clients
+	if (Output.Num() && Process->bUpdateClientsWithStdout)
+	{
+		FSwitchboardProgramStdout Packet;
+
+		Packet.Process.Uuid = Process->UUID.ToString();
+		Packet.Process.Name = Process->Name;
+		Packet.Process.Path = Process->Path;
+		Packet.Process.Caller = Process->Caller;
+		Packet.Process.Pid = Process->PID;
+
+		Packet.PartialStdoutB64 = FBase64::Encode(Output);
+
+		UE::TUniqueLock<UE::FRecursiveMutex> ConnectionsScopeLock(ConnectionsLock);
+		for (const TPair<FIPv4Endpoint, FConnectionRef>& Connection : ConnectionsByEndpoint)
+		{
+			const FIPv4Endpoint& ClientEndpoint = Connection.Key;
+			SendMessage(CreateMessage(Packet), ClientEndpoint);
+		}
+	}
+}
+
+void FSwitchboardListener::HandleRunningProcesses(TArray<TSharedPtr<FRunningProcess>>& Processes, bool bNotifyThatProgramEnded)
+{
+	SWITCHBOARD_TRACE_CPUPROFILER_EVENT_SCOPE(FSwitchboardListener::HandleRunningProcesses);
+
+	// Reads pipe and cleans up dead processes from the array.
+	for (auto Iter = Processes.CreateIterator(); Iter; ++Iter)
+	{
+		TSharedPtr<FRunningProcess> Process = *Iter;
+
+		check(Process.IsValid());
+
+		if (Process->bPendingKill)
+		{
+			continue;
+		}
+
+		if (Process->Handle.IsValid())
+		{
+			HandleStdout(Process);
+
+			if (!FPlatformProcess::IsProcRunning(Process->Handle))
+			{
+				// A final read can be necessary to avoid truncating the output.
+				HandleStdout(Process);
+
+				int32 ReturnCode = 0;
+				FPlatformProcess::GetProcReturnCode(Process->Handle, &ReturnCode);
+				UE_LOG(LogSwitchboard, Display, TEXT("Process %d (%s) exited with returncode: %d"),
+					Process->PID, *Process->Name, ReturnCode);
+
+				if (ReturnCode != 0 && Process->Output.Num() > 0)
+				{
+					FUTF8ToTCHAR StdoutConv(reinterpret_cast<const ANSICHAR*>(Process->Output.GetData()), Process->Output.Num());
+					UE_LOG(LogSwitchboard, Display, TEXT("Output:\n%.*s"), StdoutConv.Length(), StdoutConv.Get());
+				}
+
+				// Notify remote client, which implies that this is a program managed by it.
+				if (bNotifyThatProgramEnded)
+				{
+					FSwitchboardProgramEnded Packet;
+
+					Packet.Process.Uuid = Process->UUID.ToString();
+					Packet.Process.Name = Process->Name;
+					Packet.Process.Path = Process->Path;
+					Packet.Process.Caller = Process->Caller;
+					Packet.Process.Pid = Process->PID;
+					Packet.Returncode = ReturnCode;
+					Packet.StdoutB64 = FBase64::Encode(Process->Output);
+
+					{
+						UE::TUniqueLock<UE::FRecursiveMutex> ConnectionsScopeLock(ConnectionsLock);
+						for (const TPair<FIPv4Endpoint, FConnectionRef>& Connection : ConnectionsByEndpoint)
+						{
+							const FIPv4Endpoint& ClientEndpoint = Connection.Key;
+							SendMessage(CreateMessage(Packet), ClientEndpoint);
+						}
+					}
+
+					// Kill its monitor to avoid potential zombies (unless it is already pending kill)
+					{
+						TSharedPtr<FRunningProcess>* FlipModeMonitorPtr = FlipModeMonitors.FindByPredicate([&](const TSharedPtr<FRunningProcess>& FlipMonitor)
+						{
+							check(FlipMonitor.IsValid());
+							return !FlipMonitor->bPendingKill && (FlipMonitor->UUID == Process->UUID);
+						});
+
+						if (FlipModeMonitorPtr)
+						{
+							check((*FlipModeMonitorPtr).IsValid());
+
+							FSwitchboardKillTask Task(FGuid(), (*FlipModeMonitorPtr)->Recipient, (*FlipModeMonitorPtr)->UUID);
+							Task_KillProcess(Task);
+						}
+					}
+				}
+
+				FPlatformProcess::CloseProc(Process->Handle);
+				Process->ClosePipes();
+
+				Iter.RemoveCurrent();
+			}
+		}
+	}
+}
+
+bool FSwitchboardListener::SendMessage(const FString& InMessage, const FIPv4Endpoint& InEndpoint)
+{
+	SWITCHBOARD_TRACE_CPUPROFILER_EVENT_SCOPE(FSwitchboardListener::SendMessage);
+
+	if (InEndpoint == InvalidEndpoint)
+	{
+		return false;
+	}
+
+	UE::TUniqueLock<UE::FRecursiveMutex> ConnectionsScopeLock(ConnectionsLock);
+
+	if (!ConnectionsByEndpoint.Contains(InEndpoint))
+	{
+		// this happens when a client disconnects while a task it had issued is not finished
+		UE_LOG(LogSwitchboard, Verbose, TEXT("Trying to send message to disconnected client %s"), *InEndpoint.ToString());
+		return false;
+	}
+
+	FConnectionRef Connection = ConnectionsByEndpoint[InEndpoint];
+
+	UE_LOG(LogSwitchboardProtocol, Verbose, TEXT("Sending message %s"), *InMessage);
+
+	uint64 Utf8Length = FPlatformString::ConvertedLength<UTF8CHAR>(*InMessage, InMessage.Len() + 1);
+	FByteArrayRef SendArray = MakeShared<TArray<uint8>>();
+	SendArray->SetNumUninitialized(Utf8Length);
+	FPlatformString::Convert((UTF8CHAR*)SendArray->GetData(), SendArray->Num(), *InMessage, InMessage.Len() + 1);
+
+	UE::TUniqueLock<UE::FMutex> SendLock(Connection->SendLock);
+	TSharedPtr<FQuicBuffer> SendBuffer = MakeShared<FQuicBuffer>(MoveTemp(SendArray));
+	Connection->SendBuffers.Enqueue(SendBuffer);
+	QuicApi->StreamSend(Connection->QuicStream, &SendBuffer->QuicBuffer, 1, QUIC_SEND_FLAG_NONE, SendBuffer.Get());
+
+	return true;
+}
+
+void FSwitchboardListener::SendMessageFutures()
+{
+	SWITCHBOARD_TRACE_CPUPROFILER_EVENT_SCOPE(FSwitchboardListener::SendMessageFutures);
+
+	for (auto Iter = MessagesFutures.CreateIterator(); Iter; ++Iter)
+	{
+		FSwitchboardMessageFuture& MessageFuture = *Iter;
+
+		if (!MessageFuture.Future.IsReady())
+		{
+			continue;
+		}
+
+		FString Message = MessageFuture.Future.Get();
+		if (!Message.IsEmpty())
+		{
+			SendMessage(Message, MessageFuture.InEndpoint);
+		}
+
+		Iter.RemoveCurrent();
+	}
+}
diff --git a/Engine/Source/Programs/SwitchboardListener/SblCore/Private/SwitchboardMainCommon.cpp b/Engine/Source/Programs/SwitchboardListener/SblCore/Private/SwitchboardMainCommon.cpp
new file mode 100644
--- /dev/null
+++ b/Engine/Source/Programs/SwitchboardListener/SblCore/Private/SwitchboardMainCommon.cpp
@@ -0,0 +1,318 @@
+// Copyright Epic Games, Inc. All Rights Reserved.
+
+#include "SwitchboardListener.h"
+#include "SwitchboardListenerApp.h"
+#include "SwitchboardListenerVersion.h"
+
+#include "Async/TaskGraphInterfaces.h"
+#include "HAL/ExceptionHandling.h"
+#include "HAL/PlatformApplicationMisc.h"
+#include "HAL/PlatformFile.h"
+#include "HAL/PlatformFileManager.h"
+#include "HAL/ThreadingBase.h"
+#include "Interfaces/IPv4/IPv4Endpoint.h"
+#include "Interfaces/IPluginManager.h"
+#include "LaunchEngineLoop.h"
+#include "Misc/CommandLine.h"
+#include "Misc/ConfigCacheIni.h"
+#include "Misc/CoreDelegates.h"
+#include "Misc/OutputDeviceError.h"
+#include "Misc/OutputDeviceRedirector.h"
+#include "Misc/ScopeExit.h"
+#include "Modules/BuildVersion.h"
+#include "Modules/ModuleManager.h"
+#include "MsQuicRuntimeModule.h"
+#include "ProfilingDebugging/TraceAuxiliary.h"
+#include "UObject/UObjectBase.h"
+
+
+DEFINE_LOG_CATEGORY(LogSwitchboard);
+
+
+int32 InitEngine()
+{
+	const TCHAR* CommandLine = FCommandLine::Get();
+	const int32 InitResult = GEngineLoop.PreInit(CommandLine);
+	if (InitResult != 0)
+	{
+		return InitResult;
+	}
+
+	if (GLog)
+	{
+		GLog->EnableBacklog(true);
+	}
+
+	ProcessNewlyLoadedUObjects();
+	FModuleManager::Get().StartProcessingNewlyLoadedObjects();
+
+	IPluginManager::Get().LoadModulesForEnabledPlugins(ELoadingPhase::PreDefault);
+	IPluginManager::Get().LoadModulesForEnabledPlugins(ELoadingPhase::Default);
+
+	return 0;
+}
+
+
+bool InitSocketSystem()
+{
+	EModuleLoadResult LoadResult;
+	FModuleManager::Get().LoadModuleWithFailureReason(TEXT("Sockets"), LoadResult);
+
+	FIPv4Endpoint::Initialize();
+
+	return LoadResult == EModuleLoadResult::Success;
+}
+
+
+// There's a subtle potential race condition with the parent PID being recycled
+// after the parent listener has actually exited. We open a handle first, then
+// ensure we have the right process by enumerating processes to check that it's
+// actually our parent process. At that point, the open handle is unambiguous.
+FProcHandle CheckRedeploy(const FSwitchboardCommandLineOptions& Options)
+{
+	if (!Options.RedeployFromPid.IsSet())
+	{
+		return FProcHandle();
+	}
+
+	const uint32 ExpectedParentPid = Options.RedeployFromPid.GetValue();
+	const FProcHandle ExpectedParentProc = FPlatformProcess::OpenProcess(ExpectedParentPid);
+	if (!ExpectedParentProc.IsValid())
+	{
+		UE_LOG(LogSwitchboard, Error, TEXT("Couldn't get handle to redeploy parent process; you'll need to remove the old executable manually"));
+		return FProcHandle();
+	}
+
+	const uint32 CurrentPid = FPlatformProcess::GetCurrentProcessId();
+	FPlatformProcess::FProcEnumerator ProcessEnumerator;
+	while (ProcessEnumerator.MoveNext())
+	{
+		FPlatformProcess::FProcEnumInfo CurrentProcInfo = ProcessEnumerator.GetCurrent();
+		const bool bIsCurrentProc = CurrentProcInfo.GetPID() == CurrentPid;
+		const bool bHasExpectedParent = CurrentProcInfo.GetParentPID() == ExpectedParentPid;
+		if (bIsCurrentProc)
+		{
+			if (bHasExpectedParent)
+			{
+				return ExpectedParentProc;
+			}
+			else
+			{
+				UE_LOG(LogSwitchboard, Error, TEXT("Redeploy expected parent PID of %u, but found %u; you'll need to remove the old executable manually"),
+					ExpectedParentPid, CurrentProcInfo.GetParentPID());
+
+				return FProcHandle();
+			}
+		}
+	}
+
+	// Shouldn't be possible.
+	check(false);
+	UE_LOG(LogSwitchboard, Error, TEXT("Failed to enumerate our own PID!"));
+
+	return FProcHandle();
+}
+
+
+// Wait for the old listener parent process to exit and delete its executable.
+bool HandleRedeploy(FProcHandle& RedeployParentProc, uint32 RedeployParentPid)
+{
+	const FString& OldLauncherPath = FPlatformProcess::GetApplicationName(RedeployParentPid);
+
+	// Get handle to existing IPC semaphore.
+	const FString IpcSemaphoreName = FSwitchboardListener::GetIpcSemaphoreName(RedeployParentPid);
+	FPlatformProcess::FSemaphore* IpcSemaphore = FPlatformProcess::NewInterprocessSynchObject(IpcSemaphoreName, false);
+	ON_SCOPE_EXIT
+	{
+		if (IpcSemaphore)
+		{
+			FPlatformProcess::DeleteInterprocessSynchObject(IpcSemaphore);
+		}
+	};
+
+	if (IpcSemaphore == nullptr)
+	{
+		UE_LOG(LogSwitchboard, Fatal, TEXT("Error opening redeploy IPC semaphore"));
+		return false;
+	}
+
+	// This release signals the old parent launcher that we're initialized and it can shut down.
+	IpcSemaphore->Unlock();
+
+	if (FPlatformProcess::IsProcRunning(RedeployParentProc))
+	{
+		UE_LOG(LogSwitchboard, Display, TEXT("Waiting for previous listener to shut down..."));
+		int8 SecondsToWait = 5;
+		bool bStillRunning = true;
+		do
+		{
+			FPlatformProcess::Sleep(1.0f);
+			--SecondsToWait;
+			bStillRunning = FPlatformProcess::IsProcRunning(RedeployParentProc);
+		} while (bStillRunning && SecondsToWait > 0);
+
+		if (bStillRunning)
+		{
+			UE_LOG(LogSwitchboard, Warning, TEXT("Forcibly terminating previous listener"));
+			if (RedeployParentProc.IsValid())
+			{
+				FPlatformProcess::TerminateProc(RedeployParentProc);
+				FPlatformProcess::CloseProc(RedeployParentProc);
+			}
+			bStillRunning = FPlatformProcess::IsApplicationRunning(RedeployParentPid);
+		}
+
+		if (bStillRunning)
+		{
+			UE_LOG(LogSwitchboard, Error,
+				TEXT("Unable to shut down previous listener at %s; you'll need to remove the old executable manually"),
+				*OldLauncherPath);
+			return false;
+		}
+	}
+
+	IPlatformFile& PlatformFile = FPlatformFileManager::Get().GetPlatformFile();
+	if (!PlatformFile.DeleteFile(*OldLauncherPath))
+	{
+		UE_LOG(LogSwitchboard, Error,
+			TEXT("Unable to remove previous listener at %s; you may need to remove the old executable manually"),
+			*OldLauncherPath);
+		return false;
+	}
+
+	return true;
+}
+
+
+bool SwitchboardListenerMainInit(const FSwitchboardCommandLineOptions& InOptions)
+{
+	const int32 InitResult = InitEngine();
+
+	UE_LOG(LogSwitchboard, Display, TEXT("SwitchboardListener %u.%u.%u"), SBLISTENER_VERSION_MAJOR, SBLISTENER_VERSION_MINOR, SBLISTENER_VERSION_PATCH);
+
+	if (InitResult != 0)
+	{
+		UE_LOG(LogSwitchboard, Fatal, TEXT("Could not initialize engine, Error code: %d"), InitResult);
+		RequestEngineExit(TEXT("Engine init failure"));
+		return false;
+	}
+
+	UE_LOG(LogSwitchboard, Display, TEXT("Successfully initialized engine."));
+
+	if (!InitSocketSystem())
+	{
+		UE_LOG(LogSwitchboard, Fatal, TEXT("Could not initialize socket system!"));
+		RequestEngineExit(TEXT("Socket init failure"));
+		return false;
+	}
+
+	UE_LOG(LogSwitchboard, Display, TEXT("Successfully initialized socket system."));
+
+	if (!FMsQuicRuntimeModule::InitRuntime())
+	{
+		UE_LOG(LogSwitchboard, Fatal, TEXT("Could not initialize MsQuic runtime."));
+		RequestEngineExit(TEXT("MsQuic init failure"));
+		return false;
+	}
+
+	FProcHandle RedeployParentProc = CheckRedeploy(InOptions);
+	if (RedeployParentProc.IsValid())
+	{
+		UE_LOG(LogSwitchboard, Display, TEXT("Performing redeploy"));
+		HandleRedeploy(RedeployParentProc, InOptions.RedeployFromPid.GetValue());
+	}
+
+	return true;
+}
+
+
+bool SwitchboardListenerMainShutdown()
+{
+	// Cleanup. See also: FEngineLoop::AppPreExit and FEngineLoop::AppExit
+	UE_LOG(LogExit, Log, TEXT("Preparing to exit."));
+
+	FCoreDelegates::OnPreExit.Broadcast();
+	FCoreDelegates::OnExit.Broadcast();
+
+	FModuleManager::Get().UnloadModulesAtShutdown();
+
+	if (GThreadPool != nullptr)
+	{
+		GThreadPool->Destroy();
+	}
+	if (GBackgroundPriorityThreadPool != nullptr)
+	{
+		GBackgroundPriorityThreadPool->Destroy();
+	}
+	if (GIOThreadPool != nullptr)
+	{
+		GIOThreadPool->Destroy();
+	}
+
+	FTaskGraphInterface::Shutdown();
+
+#if WITH_APPLICATION_CORE
+	FPlatformApplicationMisc::TearDown();
+#endif
+	FPlatformMisc::PlatformTearDown();
+
+	if (GConfig)
+	{
+		GConfig->Exit();
+		delete GConfig;
+		GConfig = nullptr;
+	}
+
+	if (GLog)
+	{
+		GLog->TearDown();
+	}
+
+	FTextLocalizationManager::TearDown();
+	FInternationalization::TearDown();
+
+	FTraceAuxiliary::Shutdown();
+
+	return true;
+}
+
+
+extern int32 SwitchboardListenerMain();
+
+
+int32 SwitchboardListenerMainWrapper()
+{
+	int32 ExitCode;
+	if (FPlatformMisc::IsDebuggerPresent())
+	{
+		ExitCode = SwitchboardListenerMain();
+	}
+	else
+	{
+#if PLATFORM_WINDOWS && !PLATFORM_SEH_EXCEPTIONS_DISABLED
+		__try
+#endif
+		{
+			// SetCrashHandler(nullptr) sets up default behavior for Linux and Mac interfacing with CrashReportClient
+			FPlatformMisc::SetCrashHandler(nullptr);
+			GIsGuarded = true;
+			ExitCode = SwitchboardListenerMain();
+			GIsGuarded = false;
+		}
+#if PLATFORM_WINDOWS && !PLATFORM_SEH_EXCEPTIONS_DISABLED
+		__except (ReportCrash(GetExceptionInformation()))
+		{
+			if (GError)
+			{
+				GError->HandleError();
+			}
+
+			// RequestExit(bForce==True) terminates this process with a specific exit code and does not return.
+			FPlatformMisc::RequestExit(true);
+			ExitCode = 1; // Suppress compiler warning.
+		}
+#endif
+	}
+
+	return ExitCode;
+}
diff --git a/Engine/Source/Programs/SwitchboardListener/SblCore/Private/SwitchboardMessageFuture.h b/Engine/Source/Programs/SwitchboardListener/SblCore/Private/SwitchboardMessageFuture.h
new file mode 100644
--- /dev/null
+++ b/Engine/Source/Programs/SwitchboardListener/SblCore/Private/SwitchboardMessageFuture.h
@@ -0,0 +1,20 @@
+// Copyright Epic Games, Inc. All Rights Reserved.
+
+#pragma once
+
+#include "Async/Future.h"
+#include "Containers/UnrealString.h"
+#include "Interfaces/IPv4/IPv4Endpoint.h"
+#include "SwitchboardTasks.h"
+
+struct FSwitchboardMessageFuture
+{
+	/** Endpoint where this message should be sent to */
+	FIPv4Endpoint InEndpoint;
+
+	/** Equivalence hash of the original task */
+	uint32 EquivalenceHash;
+
+	/** Future containing message contents */
+	TFuture<FString> Future;
+};
\ No newline at end of file
diff --git a/Engine/Source/Programs/SwitchboardListener/SblCore/Private/SwitchboardPacket.h b/Engine/Source/Programs/SwitchboardListener/SblCore/Private/SwitchboardPacket.h
new file mode 100644
--- /dev/null
+++ b/Engine/Source/Programs/SwitchboardListener/SblCore/Private/SwitchboardPacket.h
@@ -0,0 +1,157 @@
+// Copyright Epic Games, Inc. All Rights Reserved.
+
+#pragma once
+
+#include "CoreMinimal.h"
+
+#include "SwitchboardListenerVersion.h"
+
+#include "SwitchboardPacket.generated.h"
+
+
+USTRUCT()
+struct FSwitchboardPacket
+{
+	GENERATED_BODY()
+
+	UPROPERTY()
+	FString Command;
+
+	UPROPERTY()
+	bool bAck;
+
+	UPROPERTY()
+	FString Error;
+};
+
+
+USTRUCT()
+struct FSwitchboardStateRunningProcess
+{
+	GENERATED_BODY()
+
+	UPROPERTY()
+	FString Uuid;
+
+	UPROPERTY()
+	FString Path;
+
+	UPROPERTY()
+	FString Name;
+
+	UPROPERTY()
+	FString Caller;
+
+	UPROPERTY()
+	uint32 Pid;
+};
+
+// This is an initial snapshot, sent upon connecting.
+USTRUCT()
+struct FSwitchboardStatePacket : public FSwitchboardPacket
+{
+	GENERATED_BODY()
+
+	FSwitchboardStatePacket() : FSwitchboardPacket()
+	{
+		Command = TEXT("state");
+		bAck = true;
+
+		const uint32 Major = SBLISTENER_VERSION_MAJOR;
+		const uint32 Minor = SBLISTENER_VERSION_MINOR;
+		const uint32 Patch = SBLISTENER_VERSION_PATCH;
+
+		Version = (Major << 16) | (Minor << 8) | (Patch);
+	}
+
+	UPROPERTY()
+	TArray<FSwitchboardStateRunningProcess> RunningProcesses;
+
+	UPROPERTY()
+	uint32 Version;
+
+	UPROPERTY()
+	FString OsVersionLabel;
+
+	UPROPERTY()
+	FString OsVersionLabelSub;
+
+	UPROPERTY()
+	FString OsVersionNumber;
+
+	UPROPERTY()
+	uint64 TotalPhysicalMemory;
+
+	/** The name of the platform as it appears in Engine/Binaries */
+	UPROPERTY()
+	FString PlatformBinaryDirectory;
+};
+
+USTRUCT()
+struct FSwitchboardProgramStdout : public FSwitchboardPacket
+{
+	GENERATED_BODY()
+
+	FSwitchboardProgramStdout() : FSwitchboardPacket()
+	{
+		Command = TEXT("programstdout");
+		bAck = true;
+	}
+
+	UPROPERTY()
+	FSwitchboardStateRunningProcess Process;
+
+	UPROPERTY()
+	FString PartialStdoutB64;
+};
+
+USTRUCT()
+struct FSwitchboardProgramEnded : public FSwitchboardPacket
+{
+	GENERATED_BODY()
+
+	FSwitchboardProgramEnded() : FSwitchboardPacket()
+	{
+		Command = TEXT("program ended");
+		bAck = true;
+	}
+
+	UPROPERTY()
+	FSwitchboardStateRunningProcess Process;
+
+	UPROPERTY()
+	int32 Returncode;
+
+	UPROPERTY()
+	FString StdoutB64;
+};
+
+USTRUCT()
+struct FSwitchboardProgramStarted : public FSwitchboardPacket
+{
+	GENERATED_BODY()
+
+	FSwitchboardProgramStarted() : FSwitchboardPacket()
+	{
+		Command = TEXT("program started");
+		bAck = true;
+	}
+
+	UPROPERTY()
+	FSwitchboardStateRunningProcess Process;
+};
+
+USTRUCT()
+struct FSwitchboardProgramKilled : public FSwitchboardPacket
+{
+	GENERATED_BODY()
+
+	FSwitchboardProgramKilled() : FSwitchboardPacket()
+	{
+		Command = TEXT("program killed");
+		bAck = true;
+	}
+
+	UPROPERTY()
+	FSwitchboardStateRunningProcess Process;
+};
diff --git a/Engine/Source/Programs/SwitchboardListener/SblCore/Private/SwitchboardProtocol.cpp b/Engine/Source/Programs/SwitchboardListener/SblCore/Private/SwitchboardProtocol.cpp
new file mode 100644
--- /dev/null
+++ b/Engine/Source/Programs/SwitchboardListener/SblCore/Private/SwitchboardProtocol.cpp
@@ -0,0 +1,368 @@
+// Copyright Epic Games, Inc. All Rights Reserved.
+
+#include "SwitchboardProtocol.h"
+
+#include "SwitchboardPacket.h"
+#include "SwitchboardTasks.h"
+#include "SyncStatus.h"
+
+#include "Dom/JsonValue.h"
+#include "JsonObjectConverter.h"
+#include "Logging/StructuredLog.h"
+#include "Policies/CondensedJsonPrintPolicy.h"
+#include "Serialization/JsonReader.h"
+#include "Serialization/JsonSerializer.h"
+#include "Serialization/JsonWriter.h"
+
+
+DEFINE_LOG_CATEGORY(LogSwitchboardProtocol);
+
+
+FString CreateMessage(const TMap<FString, FString>& InFields)
+{
+	FString Message;
+	TSharedRef<TJsonWriter<TCHAR, TCondensedJsonPrintPolicy<TCHAR>>> JsonWriter = TJsonWriterFactory<TCHAR, TCondensedJsonPrintPolicy<TCHAR>>::Create(&Message);
+	JsonWriter->WriteObjectStart();
+	for (const TPair<FString, FString>& Field : InFields)
+	{
+		JsonWriter->WriteValue(Field.Key, Field.Value);
+	}
+	JsonWriter->WriteObjectEnd();
+	JsonWriter->Close();
+	return Message;
+}
+
+
+FString CreateMessage(const FString& InStateDescription, bool bInState, const TMap<FString, FString>& InAdditionalFields)
+{
+	FString Message;
+	TSharedRef<TJsonWriter<TCHAR, TCondensedJsonPrintPolicy<TCHAR>>> JsonWriter = TJsonWriterFactory<TCHAR, TCondensedJsonPrintPolicy<TCHAR>>::Create(&Message);
+	JsonWriter->WriteObjectStart();
+	JsonWriter->WriteValue(InStateDescription, bInState); // TODO: Phase out this field and replace with two below because parser now needs to check all possibilities.
+	JsonWriter->WriteValue(TEXT("command"), InStateDescription);
+	JsonWriter->WriteValue(TEXT("bAck"), bInState);
+	for (const auto& Value : InAdditionalFields)
+	{
+		JsonWriter->WriteValue(Value.Key, Value.Value);
+	}
+	JsonWriter->WriteObjectEnd();
+	JsonWriter->Close();
+	return Message;
+}
+
+FString CreateTaskDeclinedMessage(const FSwitchboardTask& InTask, const FString& InErrorMessage, const TMap<FString, FString>& InAdditionalFields)
+{
+	TMap<FString, FString> AdditionalFields = InAdditionalFields;
+
+	AdditionalFields.Add(TEXT("id"), InTask.TaskID.ToString());
+	AdditionalFields.Add(TEXT("error"), InErrorMessage);
+
+	return CreateMessage(InTask.GetCommandName(), false, AdditionalFields);
+}
+
+FString CreateCommandAcceptedMessage(const FGuid& InMessageID)
+{
+	return CreateMessage(TEXT("command accepted"), true, { { TEXT("id"), InMessageID.ToString() } });
+}
+
+FString CreateCommandDeclinedMessage(const FGuid& InMessageID, const FString& InErrorMessage)
+{
+	return CreateMessage(TEXT("command accepted"), false, { { TEXT("id"), InMessageID.ToString() }, {TEXT("error"), InErrorMessage} });
+}
+
+FString CreateSyncStatusMessage(const FSyncStatus& SyncStatus)
+{
+	FString SyncStatusJsonString;
+	const bool bJsonStringOk = FJsonObjectConverter::UStructToJsonObjectString(SyncStatus, SyncStatusJsonString);
+
+	check(bJsonStringOk);
+
+	FString Message;
+
+	TSharedRef<TJsonWriter<TCHAR, TCondensedJsonPrintPolicy<TCHAR>>> JsonWriter = TJsonWriterFactory<TCHAR, TCondensedJsonPrintPolicy<TCHAR>>::Create(&Message);
+	JsonWriter->WriteObjectStart();
+	JsonWriter->WriteValue(TEXT("get sync status"), true); // TODO: Phase out this field and replace with two below because parser now needs to check all possibilities.
+	JsonWriter->WriteValue(TEXT("command"), TEXT("get sync status"));
+	JsonWriter->WriteValue(TEXT("bAck"), true);
+	JsonWriter->WriteRawJSONValue(TEXT("syncStatus"), SyncStatusJsonString);
+	JsonWriter->WriteObjectEnd();
+	JsonWriter->Close();
+
+	return Message;
+}
+
+
+FString CreateReceiveFileFromClientCompletedMessage(const FString& InDestinationPath)
+{
+	return CreateMessage(TEXT("send file complete"), true, { { TEXT("destination"), InDestinationPath } });
+}
+
+FString CreateReceiveFileFromClientFailedMessage(const FString& InDestinationPath, const FString& InError)
+{
+	return CreateMessage(TEXT("send file complete"), false, { { TEXT("destination"), InDestinationPath }, { TEXT("error"), InError } });
+}
+
+FString CreateSendFileToClientCompletedMessage(const FString& InSourcePath, const FString& InFileContent)
+{
+	return CreateMessage(TEXT("receive file complete"), true, { { TEXT("source"), InSourcePath }, { TEXT("content"), InFileContent } });
+}
+
+FString CreateSendFileToClientFailedMessage(const FString& InSourcePath, const FString& InError)
+{
+	return CreateMessage(TEXT("receive file complete"), false, { { TEXT("source"), InSourcePath }, { TEXT("error"), InError } });
+}
+
+FString CreateRedeployStatusMessage(const FGuid& InMessageID, bool bAck, const FString& Status)
+{
+	return CreateMessage(TEXT("redeploy server status"), bAck, { { TEXT("id"), InMessageID.ToString() }, { TEXT("status"), Status } });
+}
+
+
+// This logs an error with some context in the event the field is missing.
+TSharedPtr<FJsonValue> TryGetCommandRequiredField(const TSharedPtr<FJsonObject>& CommandObj, const FString& FieldName)
+{
+	TSharedPtr<FJsonValue> Field = CommandObj->TryGetField(FieldName);
+	if (!Field)
+	{
+		UE_LOG(LogSwitchboardProtocol, Error, TEXT("\"%s\" command missing required field \"%s\""),
+			*CommandObj->GetStringField(TEXT("command")).ToLower(), *FieldName);
+	}
+	return Field;
+}
+
+
+FCreateTaskResult CreateTaskFromCommand(const FString& InCommand, const FIPv4Endpoint& InEndpoint, bool bAuthenticated)
+{
+	FCreateTaskResult Result;
+
+	TSharedRef<TJsonReader<TCHAR>> Reader = FJsonStringReader::Create(InCommand);
+	TSharedPtr<FJsonObject> JsonData;
+	if (!FJsonSerializer::Deserialize(Reader, JsonData))
+	{
+		Result.Status = ECreateTaskStatus::Error_ParsingFailed;
+		return Result;
+	}
+
+	TSharedPtr<FJsonValue> CommandField = JsonData->TryGetField(TEXT("command"));
+	TSharedPtr<FJsonValue> IdField = JsonData->TryGetField(TEXT("id"));
+
+	if (!CommandField.IsValid())
+	{
+		Result.Status = ECreateTaskStatus::Error_ParsingFailed;
+		return Result;
+	}
+	
+	const FString CommandName = CommandField->AsString().ToLower();
+	Result.CommandName = CommandName;
+
+	FGuid MessageID;
+	if (!IdField.IsValid() || !FGuid::Parse(IdField->AsString(), MessageID))
+	{
+		Result.Status = ECreateTaskStatus::Error_ParsingFailed;
+		return Result;
+	}
+
+	Result.MessageID = MessageID;
+
+	// Should we echo this command in the output log?
+	{
+		TSharedPtr<FJsonValue> EchoField = JsonData->TryGetField(TEXT("bEcho"));
+		Result.bEcho = EchoField.IsValid() ? EchoField->AsBool() : true;
+	}
+
+	// Check commands allowed for unauthenticated clients first.
+	if (CommandName == FSwitchboardDisconnectTask::CommandName)
+	{
+		Result.Status = ECreateTaskStatus::Success;
+		Result.Task = MakeUnique<FSwitchboardDisconnectTask>(MessageID, InEndpoint);
+		return Result;
+	}
+	else if (CommandName == FSwitchboardKeepAliveTask::CommandName)
+	{
+		Result.Status = ECreateTaskStatus::Success;
+		Result.Task = MakeUnique<FSwitchboardKeepAliveTask>(MessageID, InEndpoint);
+		return Result;
+	}
+	else if (CommandName == FSwitchboardSetInactiveTimeoutTask::CommandName)
+	{
+		TSharedPtr<FJsonValue> SecondsField = TryGetCommandRequiredField(JsonData, TEXT("seconds"));
+
+		if (SecondsField.IsValid())
+		{
+			Result.Status = ECreateTaskStatus::Success;
+			Result.Task = MakeUnique<FSwitchboardSetInactiveTimeoutTask>(MessageID, InEndpoint, SecondsField->AsNumber());
+			return Result;
+		}
+	}
+	else if (CommandName == FSwitchboardAuthenticateTask::CommandName)
+	{
+		TSharedPtr<FJsonValue> TokenField = TryGetCommandRequiredField(JsonData, TEXT("token"));
+		if (!TokenField)
+		{
+			Result.Status = ECreateTaskStatus::Error_ParsingFailed;
+			return Result;
+		}
+
+		Result.Status = ECreateTaskStatus::Success;
+		Result.Task = MakeUnique<FSwitchboardAuthenticateTask>(MessageID, InEndpoint, TokenField->AsString());
+		return Result;
+	}
+
+	// These commands require authentication.
+	if (!bAuthenticated)
+	{
+		Result.Status = ECreateTaskStatus::Error_Unauthenticated;
+		return Result;
+	}
+
+	if (CommandName == FSwitchboardStartTask::CommandName)
+	{
+		TSharedPtr<FJsonValue> ExeField = TryGetCommandRequiredField(JsonData, TEXT("exe"));
+		TSharedPtr<FJsonValue> ArgsField = TryGetCommandRequiredField(JsonData, TEXT("args"));
+		TSharedPtr<FJsonValue> NameField = TryGetCommandRequiredField(JsonData, TEXT("name"));
+		TSharedPtr<FJsonValue> CallerField = TryGetCommandRequiredField(JsonData, TEXT("caller"));
+		TSharedPtr<FJsonValue> WorkingDirField = TryGetCommandRequiredField(JsonData, TEXT("working_dir"));
+
+		if (!ExeField || !ArgsField || !NameField || !CallerField || !WorkingDirField)
+		{
+			Result.Status = ECreateTaskStatus::Error_ParsingFailed;
+			return Result;
+		}
+
+		TUniquePtr<FSwitchboardStartTask> Task = MakeUnique<FSwitchboardStartTask>(
+			MessageID,
+			InEndpoint,
+			ExeField->AsString(),
+			ArgsField->AsString(),
+			NameField->AsString(),
+			CallerField->AsString(),
+			WorkingDirField->AsString()
+		);
+
+		if (TSharedPtr<FJsonValue> UpdateClientsWithStdoutField = JsonData->TryGetField(TEXT("bUpdateClientsWithStdout")))
+		{
+			UpdateClientsWithStdoutField->TryGetBool(Task->bUpdateClientsWithStdout);
+		}
+
+		if (TSharedPtr<FJsonValue> LockGpuClock = JsonData->TryGetField(TEXT("bLockGpuClock")))
+		{
+			LockGpuClock->TryGetBool(Task->bLockGpuClock);
+		}
+
+		if (TSharedPtr<FJsonValue> PriorityModifierField = JsonData->TryGetField(TEXT("priority_modifier")))
+		{
+			PriorityModifierField->TryGetNumber(Task->PriorityModifier);
+		}
+
+		Result.Status = ECreateTaskStatus::Success;
+		Result.Task = MoveTemp(Task);
+		return Result;
+	}
+	else if (CommandName == FSwitchboardKillTask::CommandName)
+	{
+		TSharedPtr<FJsonValue> UUIDField = TryGetCommandRequiredField(JsonData, TEXT("uuid"));
+
+		FGuid ProgramID;
+		if (UUIDField.IsValid() && FGuid::Parse(UUIDField->AsString(), ProgramID))
+		{
+			Result.Status = ECreateTaskStatus::Success;
+			Result.Task = MakeUnique<FSwitchboardKillTask>(MessageID, InEndpoint, ProgramID);
+			return Result;
+		}
+	}
+	else if (CommandName == FSwitchboardReceiveFileFromClientTask::CommandName)
+	{
+		TSharedPtr<FJsonValue> DestinationField = TryGetCommandRequiredField(JsonData, TEXT("destination"));
+		TSharedPtr<FJsonValue> FileContentField = TryGetCommandRequiredField(JsonData, TEXT("content"));
+		
+		if (!DestinationField || !FileContentField)
+		{
+			Result.Status = ECreateTaskStatus::Error_ParsingFailed;
+			return Result;
+		}
+
+		TUniquePtr<FSwitchboardReceiveFileFromClientTask> Task = MakeUnique<FSwitchboardReceiveFileFromClientTask>(
+			MessageID,
+			InEndpoint,
+			DestinationField->AsString(),
+			FileContentField->AsString()
+		);
+
+		if (TSharedPtr<FJsonValue> ForceOverwriteField = JsonData->TryGetField(TEXT("force_overwrite")))
+		{
+			ForceOverwriteField->TryGetBool(Task->bForceOverwrite);
+		}
+
+		Result.Status = ECreateTaskStatus::Success;
+		Result.Task = MoveTemp(Task);
+		return Result;
+	}
+	else if (CommandName == FSwitchboardSendFileToClientTask::CommandName)
+	{
+		TSharedPtr<FJsonValue> SourceField = TryGetCommandRequiredField(JsonData, TEXT("source"));
+
+		if (SourceField.IsValid())
+		{
+			Result.Status = ECreateTaskStatus::Success;
+			Result.Task = MakeUnique<FSwitchboardSendFileToClientTask>(MessageID, InEndpoint, SourceField->AsString());
+			return Result;
+		}
+	}
+	else if (CommandName == FSwitchboardGetSyncStatusTask::CommandName)
+	{
+		TSharedPtr<FJsonValue> UUIDField = TryGetCommandRequiredField(JsonData, TEXT("uuid"));
+
+		FGuid ProgramID;
+		if (UUIDField.IsValid() && FGuid::Parse(UUIDField->AsString(), ProgramID))
+		{
+			Result.Status = ECreateTaskStatus::Success;
+			Result.Task = MakeUnique<FSwitchboardGetSyncStatusTask>(MessageID, InEndpoint, ProgramID);
+			return Result;
+		}
+	}
+	else if (CommandName == FSwitchboardRedeployListenerTask::CommandName)
+	{		
+		TSharedPtr<FJsonValue> Sha1Field = TryGetCommandRequiredField(JsonData, TEXT("sha1"));
+		TSharedPtr<FJsonValue> FileContentField = TryGetCommandRequiredField(JsonData, TEXT("content"));
+
+		if (Sha1Field.IsValid() && FileContentField.IsValid())
+		{
+			Result.Status = ECreateTaskStatus::Success;
+			Result.Task = MakeUnique<FSwitchboardRedeployListenerTask>(MessageID, InEndpoint, Sha1Field->AsString(), FileContentField->AsString());
+			return Result;
+		}
+	}
+	else if (CommandName == FSwitchboardFreeListenerBinaryTask::CommandName)
+	{
+		Result.Status = ECreateTaskStatus::Success;
+		Result.Task = MakeUnique<FSwitchboardFreeListenerBinaryTask>(MessageID, InEndpoint);
+		return Result;
+	}
+	else if (CommandName == FSwitchboardFixExeFlagsTask::CommandName)
+	{
+		TSharedPtr<FJsonValue> UUIDField = TryGetCommandRequiredField(JsonData, TEXT("uuid"));
+
+		FGuid ProgramID;
+		if (UUIDField.IsValid() && FGuid::Parse(UUIDField->AsString(), ProgramID))
+		{
+			Result.Status = ECreateTaskStatus::Success;
+			Result.Task = MakeUnique<FSwitchboardFixExeFlagsTask>(MessageID, InEndpoint, ProgramID);
+			return Result;
+		}
+	}
+	else if (CommandName == FSwitchboardRefreshMosaicsTask::CommandName)
+	{
+		Result.Status = ECreateTaskStatus::Success;
+		Result.Task = MakeUnique<FSwitchboardRefreshMosaicsTask>(MessageID, InEndpoint);
+		return Result;
+	}
+	else if (CommandName == FSwitchboardMinimizeWindowsTask::CommandName)
+	{
+		Result.Status = ECreateTaskStatus::Success;
+		Result.Task = MakeUnique<FSwitchboardMinimizeWindowsTask>(MessageID, InEndpoint);
+		return Result;
+	}
+
+	Result.Status = ECreateTaskStatus::Error_Unhandled;
+	return Result;
+}
diff --git a/Engine/Source/Programs/SwitchboardListener/SblCore/Private/SwitchboardProtocol.h b/Engine/Source/Programs/SwitchboardListener/SblCore/Private/SwitchboardProtocol.h
new file mode 100644
--- /dev/null
+++ b/Engine/Source/Programs/SwitchboardListener/SblCore/Private/SwitchboardProtocol.h
@@ -0,0 +1,66 @@
+// Copyright Epic Games, Inc. All Rights Reserved.
+
+#pragma once
+
+#include "CoreMinimal.h"
+#include "Interfaces/IPv4/IPv4Endpoint.h"
+#include "JsonObjectConverter.h"
+#include "Logging/LogMacros.h"
+
+
+struct FSwitchboardPacket;
+struct FSwitchboardTask;
+struct FSyncStatus;
+
+
+DECLARE_LOG_CATEGORY_EXTERN(LogSwitchboardProtocol, Verbose, All);
+
+
+//~ Messages sent from Listener to Switchboard
+
+template<typename InStructType>
+FString CreateMessage(const InStructType& InStruct)
+{
+	FString Message;
+	const bool bMessageOk = FJsonObjectConverter::UStructToJsonObjectString(InStruct, Message);
+	check(bMessageOk);
+	return Message;
+}
+FString CreateMessage(const TMap<FString, FString>& InFields);
+
+FString CreateTaskDeclinedMessage(const FSwitchboardTask& InTask, const FString& InErrorMessage, const TMap<FString, FString>& InAdditionalFields);
+FString CreateCommandAcceptedMessage(const FGuid& InMessageID);
+FString CreateCommandDeclinedMessage(const FGuid& InMessageID, const FString& InErrorMessage);
+
+FString CreateReceiveFileFromClientCompletedMessage(const FString& InDestinationPath);
+FString CreateReceiveFileFromClientFailedMessage(const FString& InDestinationPath, const FString& InError);
+
+FString CreateSendFileToClientCompletedMessage(const FString& InSourcePath, const FString& InFileContent);
+FString CreateSendFileToClientFailedMessage(const FString& InSourcePath, const FString& InError);
+
+FString CreateSyncStatusMessage(const FSyncStatus& SyncStatus);
+
+FString CreateRedeployStatusMessage(const FGuid& InMessageID, bool bAck, const FString& Status);
+
+//~
+
+enum class ECreateTaskStatus
+{
+	Success,
+
+	Error_Unauthenticated,
+	Error_ParsingFailed,
+	Error_Unhandled,
+};
+
+struct FCreateTaskResult
+{
+	ECreateTaskStatus Status;
+	TUniquePtr<FSwitchboardTask> Task;
+	bool bEcho;
+
+	TOptional<FString> CommandName;
+	TOptional<FGuid> MessageID;
+};
+
+FCreateTaskResult CreateTaskFromCommand(const FString& InCommand, const FIPv4Endpoint& InEndpoint, bool bAuthenticated);
diff --git a/Engine/Source/Programs/SwitchboardListener/SblCore/Private/SwitchboardTasks.h b/Engine/Source/Programs/SwitchboardListener/SblCore/Private/SwitchboardTasks.h
new file mode 100644
--- /dev/null
+++ b/Engine/Source/Programs/SwitchboardListener/SblCore/Private/SwitchboardTasks.h
@@ -0,0 +1,322 @@
+// Copyright Epic Games, Inc. All Rights Reserved.
+
+#pragma once
+
+#include "Interfaces/IPv4/IPv4Endpoint.h"
+#include "Misc/SecureHash.h"
+#include "Templates/TypeHash.h"
+
+enum class ESwitchboardTaskType : uint8
+{
+	Authenticate,
+	Start,
+	Kill,
+	Restart,
+	ReceiveFileFromClient,
+	SendFileToClient,
+	Disconnect,
+	KeepAlive,
+	GetSyncStatus,
+	GetFlipMode,
+	FixExeFlags,
+	RedeployListener,
+	RefreshMosaics,
+	MinimizeWindows,
+	SetInactiveTimeout,
+	FreeListenerBinary
+};
+
+
+struct FSwitchboardTask
+{
+	ESwitchboardTaskType Type;
+	FGuid TaskID;
+	FIPv4Endpoint Recipient;
+
+	FSwitchboardTask(ESwitchboardTaskType InType, FGuid InTaskID, FIPv4Endpoint InRecipient)
+		: Type(InType)
+		, TaskID(InTaskID)
+		, Recipient(InRecipient)
+	{}
+
+	virtual ~FSwitchboardTask() = default;
+
+	virtual const TCHAR* GetCommandName() const = 0;
+
+	/** Calculates a hash that should be the same for equivalent Tasks, even if their TaskID is different */
+	virtual uint32 GetEquivalenceHash() const
+	{
+		return HashCombine(GetTypeHash(Type), GetTypeHash(Recipient));
+	}	
+};
+
+struct FSwitchboardAuthenticateTask : public FSwitchboardTask
+{
+	FSwitchboardAuthenticateTask(const FGuid& InTaskId, const FIPv4Endpoint& InEndpoint, const FString& InPassword)
+		: FSwitchboardTask{ ESwitchboardTaskType::Authenticate, InTaskId, InEndpoint }
+		, Password(InPassword)
+	{}
+
+	FString Password;
+
+	//~ Begin FSwitchboardTask interface
+	static constexpr const TCHAR* CommandName = TEXT("authenticate");
+	virtual const TCHAR* GetCommandName() const override { return CommandName; };
+
+	virtual uint32 GetEquivalenceHash() const override
+	{
+		return HashCombine(FSwitchboardTask::GetEquivalenceHash(), GetTypeHash(Password));
+	}
+	//~ End FSwitchboardTask interface
+};
+
+struct FSwitchboardGetSyncStatusTask : public FSwitchboardTask
+{
+	FSwitchboardGetSyncStatusTask(const FGuid& InTaskId, const FIPv4Endpoint& InEndpoint, const FGuid& InProgramID)
+		: FSwitchboardTask{ ESwitchboardTaskType::GetSyncStatus, InTaskId, InEndpoint }
+		, ProgramID(InProgramID)
+	{}
+
+	/** ID of the program that we wish to get the FlipMode of */
+	FGuid ProgramID;
+
+	//~ Begin FSwitchboardTask interface
+	static constexpr const TCHAR* CommandName = TEXT("get sync status");
+	virtual const TCHAR* GetCommandName() const override { return CommandName; };
+
+	virtual uint32 GetEquivalenceHash() const override
+	{
+		return HashCombine(FSwitchboardTask::GetEquivalenceHash(), GetTypeHash(ProgramID));
+	}
+	//~ End FSwitchboardTask interface
+};
+
+struct FSwitchboardRefreshMosaicsTask : public FSwitchboardTask
+{
+	FSwitchboardRefreshMosaicsTask(const FGuid& InTaskId, const FIPv4Endpoint& InEndpoint)
+		: FSwitchboardTask{ ESwitchboardTaskType::RefreshMosaics, InTaskId, InEndpoint }
+	{}
+
+	//~ Begin FSwitchboardTask interface
+	static constexpr const TCHAR* CommandName = TEXT("refresh mosaics");
+	virtual const TCHAR* GetCommandName() const override { return CommandName; };
+	//~ End FSwitchboardTask interface
+};
+
+struct FSwitchboardStartTask : public FSwitchboardTask
+{
+	FSwitchboardStartTask(
+		const FGuid& InTaskId, 
+		const FIPv4Endpoint& InEndpoint, 
+		const FString& InCommand, 
+		const FString& InArgs, 
+		const FString& InName, 
+		const FString& InCaller, 
+		const FString& InWorkingDir
+	)
+		: FSwitchboardTask{ ESwitchboardTaskType::Start, InTaskId, InEndpoint }
+		, Command(InCommand)
+		, Arguments(InArgs)
+		, Name(InName)
+		, Caller(InCaller)
+		, WorkingDir(InWorkingDir)
+	{
+	}
+
+	FString Command;
+	FString Arguments;
+	FString Name;
+	FString Caller;
+	FString WorkingDir;
+	bool bUpdateClientsWithStdout = false;
+	bool bLockGpuClock = false;
+	int32 PriorityModifier = 0;
+
+	//~ Begin FSwitchboardTask interface
+	static constexpr const TCHAR* CommandName = TEXT("start");
+	virtual const TCHAR* GetCommandName() const override { return CommandName; };
+
+	virtual uint32 GetEquivalenceHash() const override
+	{
+		uint32 Hash = HashCombine(FSwitchboardTask::GetEquivalenceHash(), GetTypeHash(Command));
+		Hash = HashCombine(Hash, GetTypeHash(Caller));
+		return HashCombine(Hash, GetTypeHash(Arguments));
+	}
+	//~ End FSwitchboardTask interface
+};
+
+struct FSwitchboardKillTask : public FSwitchboardTask
+{
+	FSwitchboardKillTask(const FGuid& InTaskId, const FIPv4Endpoint& InEndpoint, const FGuid& InProgramID)
+		: FSwitchboardTask{ ESwitchboardTaskType::Kill, InTaskId, InEndpoint}
+		, ProgramID(InProgramID)
+	{}
+
+	FGuid ProgramID; // unique ID of process to kill
+
+	//~ Begin FSwitchboardTask interface
+	static constexpr const TCHAR* CommandName = TEXT("kill");
+	virtual const TCHAR* GetCommandName() const override { return CommandName; };
+
+	virtual uint32 GetEquivalenceHash() const override
+	{
+		return HashCombine(FSwitchboardTask::GetEquivalenceHash(), GetTypeHash(ProgramID));
+	}
+	//~ End FSwitchboardTask interface
+};
+
+struct FSwitchboardReceiveFileFromClientTask : public FSwitchboardTask
+{
+	FSwitchboardReceiveFileFromClientTask(const FGuid& InTaskID, const FIPv4Endpoint& InEndpoint, const FString& InDestination, const FString& InContent)
+		: FSwitchboardTask{ ESwitchboardTaskType::ReceiveFileFromClient, InTaskID, InEndpoint }
+		, Destination(InDestination)
+		, FileContent(InContent)
+	{}
+
+	FString Destination;
+	FString FileContent;
+	bool bForceOverwrite = false;
+
+	//~ Begin FSwitchboardTask interface
+	static constexpr const TCHAR* CommandName = TEXT("send file");
+	virtual const TCHAR* GetCommandName() const override { return CommandName; };
+
+	virtual uint32 GetEquivalenceHash() const override
+	{
+		uint32 Hash = HashCombine(FSwitchboardTask::GetEquivalenceHash(), GetTypeHash(Destination));
+		return HashCombine(Hash, GetTypeHash(FileContent));
+	}
+	//~ End FSwitchboardTask interface
+};
+
+struct FSwitchboardRedeployListenerTask : public FSwitchboardTask
+{
+	FSwitchboardRedeployListenerTask(const FGuid& InTaskID, const FIPv4Endpoint& InEndpoint, const FString& InExpectedHashHexDigest, const FString& InContent)
+		: FSwitchboardTask{ ESwitchboardTaskType::RedeployListener, InTaskID, InEndpoint }
+		, FileContent(InContent)
+	{
+		ExpectedHash.FromString(InExpectedHashHexDigest);
+	}
+
+	FSHAHash ExpectedHash;
+	FString FileContent;
+
+	//~ Begin FSwitchboardTask interface
+	static constexpr const TCHAR* CommandName = TEXT("redeploy listener");
+	virtual const TCHAR* GetCommandName() const override { return CommandName; };
+
+	virtual uint32 GetEquivalenceHash() const override
+	{
+		uint32 Hash = HashCombine(FSwitchboardTask::GetEquivalenceHash(), GetTypeHash(ExpectedHash));
+		return HashCombine(Hash, GetTypeHash(FileContent));
+	}
+	//~ End FSwitchboardTask interface
+};
+
+struct FSwitchboardSendFileToClientTask : public FSwitchboardTask
+{
+	FSwitchboardSendFileToClientTask(const FGuid& InTaskID, const FIPv4Endpoint& InEndpoint, const FString& InSource)
+		: FSwitchboardTask{ ESwitchboardTaskType::SendFileToClient, InTaskID, InEndpoint }
+		, Source(InSource)
+	{}
+
+	FString Source;
+
+	//~ Begin FSwitchboardTask interface
+	static constexpr const TCHAR* CommandName = TEXT("receive file");
+	virtual const TCHAR* GetCommandName() const override { return CommandName; };
+
+	virtual uint32 GetEquivalenceHash() const override
+	{
+		return HashCombine(FSwitchboardTask::GetEquivalenceHash(), GetTypeHash(Source));
+	}
+	//~ End FSwitchboardTask interface
+};
+
+struct FSwitchboardFixExeFlagsTask : public FSwitchboardTask
+{
+	FSwitchboardFixExeFlagsTask(const FGuid& InTaskId, const FIPv4Endpoint& InEndpoint, const FGuid& InProgramID)
+		: FSwitchboardTask{ ESwitchboardTaskType::FixExeFlags, InTaskId, InEndpoint }
+		, ProgramID(InProgramID)
+	{}
+
+	FGuid ProgramID;
+
+	//~ Begin FSwitchboardTask interface
+	static constexpr const TCHAR* CommandName = TEXT("fixExeFlags");
+	virtual const TCHAR* GetCommandName() const override { return CommandName; };
+
+	virtual uint32 GetEquivalenceHash() const override
+	{
+		return HashCombine(FSwitchboardTask::GetEquivalenceHash(), GetTypeHash(ProgramID));
+	}
+	//~ End FSwitchboardTask interface
+};
+
+struct FSwitchboardDisconnectTask : public FSwitchboardTask
+{
+	FSwitchboardDisconnectTask(const FGuid& InTaskId, const FIPv4Endpoint& InEndpoint)
+		: FSwitchboardTask{ESwitchboardTaskType::Disconnect, InTaskId, InEndpoint}
+	{}
+
+	//~ Begin FSwitchboardTask interface
+	static constexpr const TCHAR* CommandName = TEXT("disconnect");
+	virtual const TCHAR* GetCommandName() const override { return CommandName; };
+	//~ End FSwitchboardTask interface
+};
+
+struct FSwitchboardKeepAliveTask : public FSwitchboardTask
+{
+	FSwitchboardKeepAliveTask(const FGuid& InTaskId, const FIPv4Endpoint& InEndpoint)
+		: FSwitchboardTask{ESwitchboardTaskType::KeepAlive, InTaskId, InEndpoint}
+	{}
+
+	//~ Begin FSwitchboardTask interface
+	static constexpr const TCHAR* CommandName = TEXT("keep alive");
+	virtual const TCHAR* GetCommandName() const override { return CommandName; };
+	//~ End FSwitchboardTask interface
+};
+
+struct FSwitchboardMinimizeWindowsTask : public FSwitchboardTask
+{
+	FSwitchboardMinimizeWindowsTask(const FGuid& InTaskId, const FIPv4Endpoint& InEndpoint)
+		: FSwitchboardTask{ ESwitchboardTaskType::MinimizeWindows, InTaskId, InEndpoint }
+	{}
+
+	//~ Begin FSwitchboardTask interface
+	static constexpr const TCHAR* CommandName = TEXT("minimize windows");
+	virtual const TCHAR* GetCommandName() const override { return CommandName; };
+	//~ End FSwitchboardTask interface
+};
+
+struct FSwitchboardSetInactiveTimeoutTask : public FSwitchboardTask
+{
+	FSwitchboardSetInactiveTimeoutTask(const FGuid& InTaskId, const FIPv4Endpoint& InEndpoint, double InTimeoutSeconds)
+		: FSwitchboardTask{ ESwitchboardTaskType::SetInactiveTimeout, InTaskId, InEndpoint }
+		, TimeoutSeconds(InTimeoutSeconds)
+	{}
+
+	double TimeoutSeconds;
+
+	//~ Begin FSwitchboardTask interface
+	static constexpr const TCHAR* CommandName = TEXT("set inactive timeout");
+	virtual const TCHAR* GetCommandName() const override { return CommandName; };
+
+	virtual uint32 GetEquivalenceHash() const override
+	{
+		return HashCombine(FSwitchboardTask::GetEquivalenceHash(), GetTypeHash(TimeoutSeconds));
+	}
+	//~ End FSwitchboardTask interface
+};
+
+struct FSwitchboardFreeListenerBinaryTask : public FSwitchboardTask
+{
+	FSwitchboardFreeListenerBinaryTask(const FGuid& InTaskID, const FIPv4Endpoint& InEndpoint)
+		: FSwitchboardTask{ ESwitchboardTaskType::FreeListenerBinary, InTaskID, InEndpoint }
+	{}
+
+	//~ Begin FSwitchboardTask interface
+	static constexpr const TCHAR* CommandName = TEXT("free binary");
+	virtual const TCHAR* GetCommandName() const override { return CommandName; };
+	//~ End FSwitchboardTask interface
+};
diff --git a/Engine/Source/Programs/SwitchboardListener/SblCore/Private/SyncStatus.h b/Engine/Source/Programs/SwitchboardListener/SblCore/Private/SyncStatus.h
new file mode 100644
--- /dev/null
+++ b/Engine/Source/Programs/SwitchboardListener/SblCore/Private/SyncStatus.h
@@ -0,0 +1,206 @@
+// Copyright Epic Games, Inc. All Rights Reserved.
+
+#pragma once
+
+#include "CoreMinimal.h"
+
+#include "SyncStatus.generated.h"
+
+
+USTRUCT()
+struct FSyncGpu
+{
+	GENERATED_BODY()
+
+	UPROPERTY()
+	bool bIsSynced;
+
+	UPROPERTY()
+	int32 Connector;
+};
+
+USTRUCT()
+struct FSyncDisplay
+{
+	GENERATED_BODY()
+
+	UPROPERTY()
+	FName SyncState;
+
+	UPROPERTY()
+	uint32 Bpc;
+
+	UPROPERTY()
+	uint8 ColorFormat;
+
+	UPROPERTY()
+	int32 Depth;
+
+};
+
+
+USTRUCT()
+struct FSyncStatusParams
+{
+	GENERATED_BODY()
+
+	UPROPERTY()
+	uint32 RefreshRate;
+
+	UPROPERTY()
+	uint32 HouseSyncIncoming;
+
+	UPROPERTY()
+	bool bHouseSync;
+
+	/** Per nvapi.h comment: "means that this P2061 board receives input from another P2061 board" */
+	UPROPERTY()
+	bool bInternalSecondary;
+};
+
+USTRUCT()
+struct FSyncDelay
+{
+	GENERATED_BODY()
+
+	UPROPERTY()
+	uint32 NumLines;
+
+	UPROPERTY()
+	uint32 NumPixels;
+
+	UPROPERTY()
+	uint32 MaxLines;
+
+	UPROPERTY()
+	uint32 MinPixels;
+};
+
+USTRUCT()
+struct FSyncControlParams
+{
+	GENERATED_BODY()
+
+	UPROPERTY()
+	int32 Polarity;
+
+	UPROPERTY()
+	int32 VMode;
+
+	UPROPERTY()
+	uint32 Interval;
+
+	UPROPERTY()
+	int32 Source;
+
+	UPROPERTY()
+	bool bInterlaced;
+
+	UPROPERTY()
+	bool bSyncSourceIsOutput;
+
+	UPROPERTY()
+	FSyncDelay SyncSkew;
+
+	UPROPERTY()
+	FSyncDelay StartupDelay;
+};
+
+USTRUCT()
+struct FSyncTopo
+{
+	GENERATED_BODY()
+
+	UPROPERTY()
+	TArray<FSyncGpu> SyncGpus;
+
+	UPROPERTY()
+	TArray<FSyncDisplay> SyncDisplays;
+
+	UPROPERTY()
+	FSyncStatusParams SyncStatusParams;
+
+	UPROPERTY()
+	FSyncControlParams SyncControlParams;
+};
+
+USTRUCT()
+struct FMosaicDisplaySettings
+{
+	GENERATED_BODY()
+
+	UPROPERTY()
+	uint32 Width;
+
+	UPROPERTY()
+	uint32 Height;
+
+	UPROPERTY()
+	uint32 Bpp;
+
+	UPROPERTY()
+	uint32 Freq;
+};
+
+USTRUCT()
+struct FMosaicTopo
+{
+	GENERATED_BODY()
+
+	UPROPERTY()
+	uint32 Rows;
+
+	UPROPERTY()
+	uint32 Columns;
+
+	UPROPERTY()
+	uint32 DisplayCount;
+
+	UPROPERTY()
+	FMosaicDisplaySettings DisplaySettings;
+};
+
+USTRUCT()
+struct FSyncStatus
+{
+	GENERATED_BODY()
+
+	UPROPERTY()
+	TArray<FSyncTopo> SyncTopos;
+
+	UPROPERTY()
+	TArray<FMosaicTopo> MosaicTopos;
+
+	UPROPERTY()
+	TArray<FString> FlipModeHistory;
+
+	UPROPERTY()
+	TArray<FString> ProgramLayers;
+
+	UPROPERTY()
+	uint32 DriverVersion;
+
+	UPROPERTY()
+	FString DriverBranch;
+
+	UPROPERTY()
+	FString Taskbar;
+
+	UPROPERTY()
+	uint32 PidInFocus;
+
+	UPROPERTY()
+	TArray<int8> CpuUtilization;
+
+	UPROPERTY()
+	uint64 AvailablePhysicalMemory;
+
+	UPROPERTY()
+	TArray<int8> GpuUtilization;
+
+	UPROPERTY()
+	TArray<int32> GpuCoreClocksKhz;
+
+	UPROPERTY()
+	TArray<int32> GpuTemperature;
+};
diff --git a/Engine/Source/Programs/SwitchboardListener/SblCore/Private/Windows/Resources/Resource.h b/Engine/Source/Programs/SwitchboardListener/SblCore/Private/Windows/Resources/Resource.h
new file mode 100644
--- /dev/null
+++ b/Engine/Source/Programs/SwitchboardListener/SblCore/Private/Windows/Resources/Resource.h
@@ -0,0 +1,24 @@
+// Copyright Epic Games, Inc. All Rights Reserved.
+
+//{{NO_DEPENDENCIES}}
+// Microsoft Visual C++ generated include file.
+// Used by Resource.rc
+//
+#define IDD_SELECTBUILD                 101
+#define IDD_ERRORDIALOG                 105
+#define IDC_BROWSE                      1001
+#define IDC_BUILDLIST                   1002
+#define IDC_ERRORLOGTEXT                1005
+#define IDC_ERRORMESSAGE                1006
+
+// Next default values for new objects
+// 
+#ifdef APSTUDIO_INVOKED
+#ifndef APSTUDIO_READONLY_SYMBOLS
+#define _APS_NEXT_RESOURCE_VALUE        106
+#define _APS_NEXT_COMMAND_VALUE         40001
+#define _APS_NEXT_CONTROL_VALUE         1007
+#define _APS_NEXT_SYMED_VALUE           101
+#endif
+#endif
+
diff --git a/Engine/Source/Programs/SwitchboardListener/SblCore/Private/Windows/Resources/Resource.rc b/Engine/Source/Programs/SwitchboardListener/SblCore/Private/Windows/Resources/Resource.rc
new file mode 100644
--- /dev/null
+++ b/Engine/Source/Programs/SwitchboardListener/SblCore/Private/Windows/Resources/Resource.rc
@@ -0,0 +1,116 @@
+// Microsoft Visual C++ generated resource script.
+//
+#include "resource.h"
+
+#define APSTUDIO_READONLY_SYMBOLS
+/////////////////////////////////////////////////////////////////////////////
+//
+// Generated from the TEXTINCLUDE 2 resource.
+//
+#include "winres.h"
+#include "Resource.rc2"
+
+/////////////////////////////////////////////////////////////////////////////
+#undef APSTUDIO_READONLY_SYMBOLS
+
+/////////////////////////////////////////////////////////////////////////////
+// English (United States) resources
+
+#if !defined(AFX_RESOURCE_DLL) || defined(AFX_TARG_ENU)
+LANGUAGE LANG_ENGLISH, SUBLANG_ENGLISH_US
+
+#ifdef APSTUDIO_INVOKED
+/////////////////////////////////////////////////////////////////////////////
+//
+// TEXTINCLUDE
+//
+
+1 TEXTINCLUDE 
+BEGIN
+    "resource.h\0"
+END
+
+2 TEXTINCLUDE 
+BEGIN
+    "#include ""winres.h""\r\n"
+    "#include ""Resource.rc2""\r\n"
+    "\0"
+END
+
+3 TEXTINCLUDE 
+BEGIN
+    "\r\n"
+    "\0"
+END
+
+#endif    // APSTUDIO_INVOKED
+
+
+/////////////////////////////////////////////////////////////////////////////
+//
+// Dialog
+//
+
+IDD_SELECTBUILD DIALOGEX 0, 0, 301, 47
+STYLE DS_SYSMODAL | DS_SETFONT | DS_MODALFRAME | DS_FIXEDSYS | DS_CENTER | WS_POPUP | WS_VISIBLE | WS_CAPTION | WS_SYSMENU
+CAPTION "Select Unreal Engine Version"
+FONT 8, "MS Shell Dlg", 400, 0, 0x1
+BEGIN
+    DEFPUSHBUTTON   "OK",IDOK,191,26,50,14
+    PUSHBUTTON      "Cancel",IDCANCEL,244,26,50,14
+    PUSHBUTTON      "...",IDC_BROWSE,270,6,24,14
+    COMBOBOX        IDC_BUILDLIST,7,7,259,135,CBS_DROPDOWNLIST | WS_VSCROLL | WS_TABSTOP
+END
+
+IDD_ERRORDIALOG DIALOGEX 0, 0, 541, 166
+STYLE DS_SETFONT | DS_MODALFRAME | DS_SETFOREGROUND | DS_FIXEDSYS | DS_CENTER | WS_POPUP | WS_CAPTION | WS_SYSMENU
+CAPTION "Error"
+FONT 8, "MS Shell Dlg", 400, 0, 0x1
+BEGIN
+    DEFPUSHBUTTON   "OK",IDOK,484,145,50,14
+    LTEXT           "Static",IDC_ERRORMESSAGE,7,7,527,8
+    EDITTEXT        IDC_ERRORLOGTEXT,7,19,527,124,ES_MULTILINE | ES_AUTOVSCROLL | ES_AUTOHSCROLL | ES_READONLY | WS_VSCROLL | WS_HSCROLL
+END
+
+
+/////////////////////////////////////////////////////////////////////////////
+//
+// DESIGNINFO
+//
+
+#ifdef APSTUDIO_INVOKED
+GUIDELINES DESIGNINFO
+BEGIN
+    IDD_SELECTBUILD, DIALOG
+    BEGIN
+        LEFTMARGIN, 7
+        RIGHTMARGIN, 294
+        TOPMARGIN, 7
+        BOTTOMMARGIN, 40
+    END
+
+    IDD_ERRORDIALOG, DIALOG
+    BEGIN
+        LEFTMARGIN, 7
+        RIGHTMARGIN, 534
+        TOPMARGIN, 7
+        BOTTOMMARGIN, 159
+    END
+END
+#endif    // APSTUDIO_INVOKED
+
+#endif    // English (United States) resources
+/////////////////////////////////////////////////////////////////////////////
+
+
+
+#ifndef APSTUDIO_INVOKED
+/////////////////////////////////////////////////////////////////////////////
+//
+// Generated from the TEXTINCLUDE 3 resource.
+//
+
+
+/////////////////////////////////////////////////////////////////////////////
+#endif    // not APSTUDIO_INVOKED
+
diff --git a/Engine/Source/Programs/SwitchboardListener/SblCore/Private/Windows/Resources/Resource.rc2 b/Engine/Source/Programs/SwitchboardListener/SblCore/Private/Windows/Resources/Resource.rc2
new file mode 100644
--- /dev/null
+++ b/Engine/Source/Programs/SwitchboardListener/SblCore/Private/Windows/Resources/Resource.rc2
@@ -0,0 +1,101 @@
+// Copyright Epic Games, Inc. All Rights Reserved.
+
+#include "../../../../../../Runtime/Launch/Resources/Windows/resource.h"
+#include "../../../../../../Runtime/Core/Public/Misc/CoreMiscDefines.h"
+#include "../../../../../../Runtime/Launch/Resources/Version.h"
+
+#include "../../../Internal/SwitchboardListenerVersion.h"
+
+#define APSTUDIO_READONLY_SYMBOLS
+#include <windows.h>
+#undef APSTUDIO_READONLY_SYMBOLS
+
+// Various strings used for project resources
+#ifdef PROJECT_COMPANY_NAME
+	#define BUILD_PROJECT_COMPANY_NAME PREPROCESSOR_TO_STRING(PROJECT_COMPANY_NAME)
+#else
+	#define BUILD_PROJECT_COMPANY_NAME EPIC_COMPANY_NAME
+#endif
+
+#ifdef PROJECT_COPYRIGHT_STRING
+	#define BUILD_PROJECT_COPYRIGHT_STRING PREPROCESSOR_TO_STRING(PROJECT_COPYRIGHT_STRING)
+#else
+	#define BUILD_PROJECT_COPYRIGHT_STRING EPIC_COPYRIGHT_STRING
+#endif
+
+#ifdef PROJECT_PRODUCT_NAME
+	#define BUILD_PROJECT_PRODUCT_NAME PREPROCESSOR_TO_STRING(PROJECT_PRODUCT_NAME)
+#else
+	#define BUILD_PROJECT_PRODUCT_NAME PREPROCESSOR_TO_STRING(UE_APP_NAME)
+#endif
+
+#ifdef BUILD_VERSION
+	#define BUILD_PROJECT_PRODUCT_VERSION PREPROCESSOR_TO_STRING(BUILD_VERSION)
+#else
+	#define BUILD_PROJECT_PRODUCT_VERSION ENGINE_VERSION_STRING
+#endif
+
+#ifdef PROJECT_PRODUCT_IDENTIFIER
+	#define BUILD_PROJECT_PRODUCT_IDENTIFIER PREPROCESSOR_TO_STRING(PROJECT_PRODUCT_IDENTIFIER)
+#else
+	#define BUILD_PROJECT_PRODUCT_IDENTIFIER EPIC_PRODUCT_IDENTIFIER
+#endif
+
+/////////////////////////////////////////////////////////////////////////////
+// English (United States) resources
+
+#if !defined(AFX_RESOURCE_DLL) || defined(AFX_TARG_ENU)
+LANGUAGE LANG_ENGLISH, SUBLANG_ENGLISH_US
+#pragma code_page(1252)
+
+/////////////////////////////////////////////////////////////////////////////
+//
+// Version
+//
+
+VS_VERSION_INFO VERSIONINFO
+ FILEVERSION SBLISTENER_VERSION_MAJOR,SBLISTENER_VERSION_MINOR,SBLISTENER_VERSION_PATCH,0
+ PRODUCTVERSION ENGINE_MAJOR_VERSION,ENGINE_MINOR_VERSION,ENGINE_PATCH_VERSION,0
+ FILEFLAGSMASK 0x17L
+#ifdef _DEBUG
+ FILEFLAGS 0x1L
+#else
+ FILEFLAGS 0x0L
+#endif
+ FILEOS 0x4L
+ FILETYPE 0x2L
+ FILESUBTYPE 0x0L
+BEGIN
+	BLOCK "StringFileInfo"
+	BEGIN
+		BLOCK "040904b0"
+		BEGIN
+			VALUE "CompanyName", BUILD_PROJECT_COMPANY_NAME
+			VALUE "LegalCopyright", BUILD_PROJECT_COPYRIGHT_STRING
+			VALUE "ProductName", BUILD_PROJECT_PRODUCT_NAME
+			VALUE "ProductVersion", BUILD_PROJECT_PRODUCT_VERSION
+			VALUE "FileDescription", BUILD_PROJECT_PRODUCT_NAME
+			VALUE "InternalName", BUILD_PROJECT_PRODUCT_IDENTIFIER
+#ifdef ORIGINAL_FILE_NAME
+			VALUE "OriginalFilename", PREPROCESSOR_TO_STRING(ORIGINAL_FILE_NAME)
+#endif
+		END
+	END
+	BLOCK "VarFileInfo"
+	BEGIN
+		VALUE "Translation", 0x409, 1200
+	END
+END
+
+/////////////////////////////////////////////////////////////////////////////
+//
+// Icon
+//
+
+// Icon with lowest ID value placed first to ensure application icon
+// remains consistent on all systems.
+
+IDICON_UEGame			ICON                    "Icon.ico"
+
+#endif    // English (United States) resources
+/////////////////////////////////////////////////////////////////////////////
diff --git a/Engine/Source/Programs/SwitchboardListener/SblCore/Private/Windows/WindowsCpuUtilizationMonitor.cpp b/Engine/Source/Programs/SwitchboardListener/SblCore/Private/Windows/WindowsCpuUtilizationMonitor.cpp
new file mode 100644
--- /dev/null
+++ b/Engine/Source/Programs/SwitchboardListener/SblCore/Private/Windows/WindowsCpuUtilizationMonitor.cpp
@@ -0,0 +1,106 @@
+// Copyright Epic Games, Inc. All Rights Reserved.
+
+#include "Windows/WindowsCpuUtilizationMonitor.h"
+#include "HAL/PlatformMisc.h"
+#include "HAL/PlatformTime.h"
+
+
+FWindowsCpuUtilizationMonitor::FWindowsCpuUtilizationMonitor()
+{
+	const int32 NumCores = FPlatformMisc::NumberOfCores();
+
+	CounterHandles.SetNum(NumCores);
+	CachedCoreUtilization.SetNum(NumCores);
+
+	PDH_STATUS PdhStatus = ::PdhOpenQuery(nullptr, 0, &QueryHandle);
+	if (PdhStatus != ERROR_SUCCESS)
+	{
+		UE_LOG(LogCpuUtilizationMonitor, Warning, TEXT("PdhOpenQuery failed. Error code: %d"), PdhStatus);
+		return;
+	}
+
+	for (int32 CoreIdx = 0; CoreIdx < NumCores; ++CoreIdx)
+	{
+		const FString CounterName = FString::Printf(TEXT("\\Processor(%d)\\%% Processor Time"), CoreIdx);
+		PdhStatus = ::PdhAddEnglishCounter(QueryHandle, *CounterName, 0, &CounterHandles[CoreIdx]);
+		if (PdhStatus != ERROR_SUCCESS)
+		{
+			UE_LOG(LogCpuUtilizationMonitor, Warning, TEXT("PdhAddEnglishCounter failed. Error code: %d"), PdhStatus);
+			return;
+		}
+	}
+
+	// Initial condition: ensures we've collected twice before the first attempt to get values.
+	PdhStatus = ::PdhCollectQueryData(QueryHandle);
+	if (PdhStatus != ERROR_SUCCESS)
+	{
+		UE_LOG(LogCpuUtilizationMonitor, Warning, TEXT("PdhCollectQueryData failed. Error code: %d"), PdhStatus);
+		return;
+	}
+
+	bIsInitialized = true;
+}
+
+FWindowsCpuUtilizationMonitor::~FWindowsCpuUtilizationMonitor()
+{
+	if (QueryHandle != INVALID_HANDLE_VALUE)
+	{
+		PDH_STATUS PdhStatus = ::PdhCloseQuery(QueryHandle);
+		if (PdhStatus != ERROR_SUCCESS)
+		{
+			UE_LOG(LogCpuUtilizationMonitor, Warning, TEXT("PdhCloseQuery failed. Error code: %d"), PdhStatus);
+		}
+	}
+}
+
+bool FWindowsCpuUtilizationMonitor::GetPerCoreUtilization(TArray<int8>& OutCoreUtilization)
+{
+	if (!IsInitialized())
+	{
+		return false;
+	}
+
+	// Rate limiting. MSDN recommends <= 1 Hz.
+	const double MinQueryIntervalSec = 1.0;
+	const double TimeNow = FPlatformTime::Seconds();
+	if (TimeNow - LastQueryTime >= MinQueryIntervalSec)
+	{
+		if (!QueryUpdatedUtilization())
+		{
+			return false;
+		}
+	}
+
+	OutCoreUtilization = CachedCoreUtilization;
+
+	return true;
+}
+
+bool FWindowsCpuUtilizationMonitor::QueryUpdatedUtilization()
+{
+	PDH_STATUS PdhStatus = ::PdhCollectQueryData(QueryHandle);
+	if (PdhStatus != ERROR_SUCCESS)
+	{
+		UE_LOG(LogCpuUtilizationMonitor, Warning, TEXT("PdhCollectQueryData failed. Error code: %d"), PdhStatus);
+		return false;
+	}
+
+	const int32 NumHandles = CounterHandles.Num();
+	for (int32 HandleIdx = 0; HandleIdx < NumHandles; ++HandleIdx)
+	{
+		// Use PDH_FMT_DOULE here because PDH_FMT_LONG truncates instead of rounding.
+		PDH_FMT_COUNTERVALUE CounterValue;
+		PdhStatus = ::PdhGetFormattedCounterValue(CounterHandles[HandleIdx], PDH_FMT_DOUBLE, 0, &CounterValue);
+		if (PdhStatus != ERROR_SUCCESS)
+		{
+			UE_LOG(LogCpuUtilizationMonitor, Warning, TEXT("PdhGetFormattedCounterValue failed. Error code: %d"), PdhStatus);
+			return false;
+		}
+
+		CachedCoreUtilization[HandleIdx] = FMath::RoundToInt(CounterValue.doubleValue);
+	}
+
+	LastQueryTime = FPlatformTime::Seconds();
+
+	return true;
+}
diff --git a/Engine/Source/Programs/SwitchboardListener/SblCore/SblCore.Build.cs b/Engine/Source/Programs/SwitchboardListener/SblCore/SblCore.Build.cs
new file mode 100644
--- /dev/null
+++ b/Engine/Source/Programs/SwitchboardListener/SblCore/SblCore.Build.cs
@@ -0,0 +1,56 @@
+// Copyright Epic Games, Inc. All Rights Reserved.
+
+using UnrealBuildTool;
+
+public class SblCore : ModuleRules
+{
+	public SblCore(ReadOnlyTargetRules Target) : base(Target)
+	{
+		PublicDependencyModuleNames.AddRange(
+			new string[]
+			{
+				"ApplicationCore", // for LaunchEngineLoop.cpp dependency
+				"Projects", // for LaunchEngineLoop.cpp dependency
+				"Core",
+				"CoreUObject",
+				"Networking",
+				"MsQuic",
+				"SwitchboardCommon",
+			}
+		);
+
+		PrivateDependencyModuleNames.AddRange(
+			new string[] {
+				"Json",
+				"JsonUtilities",
+				"MsQuicRuntime",
+				"TraceLog",
+			}
+		);
+
+		PrivateIncludePathModuleNames.AddRange(
+			new string[] {
+				"Launch",
+			}
+		);
+
+		if (Target.Platform == UnrealTargetPlatform.Win64)
+		{
+			PrivateDependencyModuleNames.AddRange(
+				new string[] {
+					"NVAPI",
+				}
+			);
+
+			PublicSystemLibraries.Add("Pdh.lib");
+		}
+
+		DynamicallyLoadedModuleNames.AddRange(
+			new string[] {
+				"Sockets",
+			}
+		);
+
+		AddEngineThirdPartyPrivateStaticDependencies(Target, "OpenSSL");
+	}
+}
diff --git a/Engine/Source/Programs/SwitchboardListener/SblSlate/Internal/SblMainWindow.h b/Engine/Source/Programs/SwitchboardListener/SblSlate/Internal/SblMainWindow.h
new file mode 100644
--- /dev/null
+++ b/Engine/Source/Programs/SwitchboardListener/SblSlate/Internal/SblMainWindow.h
@@ -0,0 +1,62 @@
+// Copyright Epic Games, Inc. All Rights Reserved.
+
+#pragma once
+
+#include "Input/Reply.h"
+#include "SwitchboardListenerAutolaunch.h"
+#include "Templates/SharedPointer.h"
+
+
+enum class ECheckBoxState : uint8;
+class FSwitchboardListener;
+class SEditableTextBox;
+class SWidget;
+class SWidgetSwitcher;
+class SWindow;
+class UToolMenu;
+
+
+class FSwitchboardListenerMainWindow
+{
+	enum class EPanelIndices
+	{
+		MainPanel = 0,
+		PasswordPanel,
+	};
+
+public:
+	FSwitchboardListenerMainWindow(FSwitchboardListener& InListener);
+
+private:
+	void OnInit();
+	void OnShutdown();
+	void OnTick();
+
+private:
+	void CustomizeToolMenus();
+	void CustomizeToolMenus_AddGeneralSection(UToolMenu* InMenu);
+	void CustomizeToolMenus_AddPasswordSection(UToolMenu* InMenu);
+	void CustomizeToolMenus_AddDevelopmentSection(UToolMenu* InMenu);
+
+	TSharedRef<SWidget> CreateRootSwitcher();
+	TSharedRef<SWidget> CreateSetPasswordPanel();
+	TSharedRef<SWidget> CreateMainPanel();
+	TSharedRef<SWidget> CreateOutputLog();
+
+	FReply OnSetPasswordClicked();
+
+private:
+	FSwitchboardListener& Listener;
+
+	TSharedPtr<SWindow> RootWindow;
+	TSharedPtr<SWidgetSwitcher> PanelSwitcher;
+
+	TSharedPtr<SEditableTextBox> PasswordTextBox;
+
+	bool bPasswordRevealedInEditBox = false;
+	bool bPasswordRevealedInMenu = false;
+
+#if SWITCHBOARD_LISTENER_AUTOLAUNCH
+	ECheckBoxState CachedAutolaunchEnabled;
+#endif
+};
diff --git a/Engine/Source/Programs/SwitchboardListener/SblSlate/Private/SblMainWindow.cpp b/Engine/Source/Programs/SwitchboardListener/SblSlate/Private/SblMainWindow.cpp
new file mode 100644
--- /dev/null
+++ b/Engine/Source/Programs/SwitchboardListener/SblSlate/Private/SblMainWindow.cpp
@@ -0,0 +1,538 @@
+// Copyright Epic Games, Inc. All Rights Reserved.
+
+#include "SblMainWindow.h"
+#include "SwitchboardListener.h"
+#include "SwitchboardListenerApp.h"
+
+#include "Framework/Application/SlateApplication.h"
+#include "Framework/Application/SWindowTitleBar.h"
+#include "Framework/Docking/LayoutService.h"
+#include "Framework/MultiBox/MultiBoxBuilder.h"
+#include "Framework/Notifications/NotificationManager.h"
+#include "HAL/PlatformApplicationMisc.h"
+#include "ISlateReflectorModule.h"
+#include "Logging/StructuredLog.h"
+#include "Misc/ConfigCacheIni.h"
+#include "Modules/ModuleManager.h"
+#include "OutputLogCreationParams.h"
+#include "OutputLogModule.h"
+#include "StandaloneRenderer.h"
+#include "Styling/StyleColors.h"
+#include "ToolMenus.h"
+#include "Widgets/Docking/SDockTab.h"
+#include "Widgets/Images/SImage.h"
+#include "Widgets/Input/SButton.h"
+#include "Widgets/Input/SCheckBox.h"
+#include "Widgets/Input/SEditableTextBox.h"
+#include "Widgets/Layout/SSpacer.h"
+#include "Widgets/Layout/SWidgetSwitcher.h"
+
+
+#define LOCTEXT_NAMESPACE "SwitchboardListener"
+
+
+class SSblWindow : public SWindow
+{
+public:
+	void Construct(const FArguments& InArgs)
+	{
+		FDisplayMetrics DisplayMetrics;
+		FSlateApplication::Get().GetDisplayMetrics(DisplayMetrics);
+		const float DPIScaleFactor = FPlatformApplicationMisc::GetDPIScaleFactorAtPoint(DisplayMetrics.PrimaryDisplayWorkAreaRect.Left, DisplayMetrics.PrimaryDisplayWorkAreaRect.Top);
+
+		SWindow::Construct(SWindow::FArguments()
+			.Title(LOCTEXT("WindowTitle", "Switchboard Listener"))
+			.CreateTitleBar(true)
+			.SupportsMaximize(true)
+			.SupportsMinimize(true)
+			.IsInitiallyMaximized(false)
+			.IsInitiallyMinimized(false)
+			.SizingRule(ESizingRule::UserSized)
+			.AutoCenter(EAutoCenter::PreferredWorkArea)
+			.ClientSize(FVector2D(1000.0f * DPIScaleFactor, 500.0f * DPIScaleFactor))
+			.AdjustInitialSizeAndPositionForDPIScale(false)
+		);
+	}
+};
+
+
+FSwitchboardListenerMainWindow::FSwitchboardListenerMainWindow(FSwitchboardListener& InListener)
+	: Listener(InListener)
+	, CachedAutolaunchEnabled(ECheckBoxState::Unchecked)
+{
+	Listener.OnInit().AddRaw(this, &FSwitchboardListenerMainWindow::OnInit);
+	Listener.OnShutdown().AddRaw(this, &FSwitchboardListenerMainWindow::OnShutdown);
+	Listener.OnTick().AddRaw(this, &FSwitchboardListenerMainWindow::OnTick);
+}
+
+
+void FSwitchboardListenerMainWindow::OnInit()
+{
+	FSlateApplication::InitializeAsStandaloneApplication(GetStandardStandaloneRenderer());
+	FSlateApplication::InitHighDPI(true);
+
+	FModuleManager::Get().LoadModuleChecked("SlateCore");
+	FModuleManager::Get().LoadModuleChecked("OutputLog");
+
+	if (IConsoleVariable* HideConsoleCVar = IConsoleManager::Get().FindConsoleVariable(TEXT("OutputLogModule.HideConsole")); ensure(HideConsoleCVar))
+	{
+		HideConsoleCVar->Set(true);
+	}
+
+	CustomizeToolMenus();
+
+	const FText ApplicationTitle = LOCTEXT("AppTitle", "Switchboard Listener");
+	FGlobalTabmanager::Get()->SetApplicationTitle(ApplicationTitle);
+
+	RootWindow = SNew(SSblWindow);
+
+	constexpr bool bShowRootWindowImmediately = true;
+	FSlateApplication::Get().AddWindow(RootWindow.ToSharedRef(), bShowRootWindowImmediately);
+	FGlobalTabmanager::Get()->SetRootWindow(RootWindow.ToSharedRef());
+	FGlobalTabmanager::Get()->SetAllowWindowMenuBar(true);
+	FSlateNotificationManager::Get().SetRootWindow(RootWindow.ToSharedRef());
+
+	RootWindow->SetContent(CreateRootSwitcher());
+
+	if (Listener.GetAuthPassword().IsEmpty())
+	{
+		PanelSwitcher->SetActiveWidgetIndex(static_cast<int32>(EPanelIndices::PasswordPanel));
+	}
+	else
+	{
+		if (ensure(Listener.StartListening()))
+		{
+			UE_LOGFMT(LogSwitchboard, Display, "Started listening");
+		}
+		else
+		{
+			UE_LOGFMT(LogSwitchboard, Error, "StartListening failed");
+		}
+	}
+
+	RootWindow->ShowWindow();
+	constexpr bool bForceWindowToFront = true;
+	RootWindow->BringToFront(bForceWindowToFront);
+}
+
+
+void FSwitchboardListenerMainWindow::OnTick()
+{
+	FSlateApplication::Get().PumpMessages();
+	FSlateApplication::Get().Tick();
+}
+
+
+void FSwitchboardListenerMainWindow::OnShutdown()
+{
+	RootWindow.Reset();
+	FSlateApplication::Shutdown();
+}
+
+
+void FSwitchboardListenerMainWindow::CustomizeToolMenus()
+{
+	UToolMenus* ToolMenus = UToolMenus::Get();
+	check(ToolMenus);
+
+	static const FName MenuName = "OutputLog.SettingsMenu";
+	UToolMenu* Menu = ToolMenus->ExtendMenu(MenuName);
+
+	CustomizeToolMenus_AddGeneralSection(Menu);
+	CustomizeToolMenus_AddPasswordSection(Menu);
+	CustomizeToolMenus_AddDevelopmentSection(Menu);
+
+	FToolMenuSection* DefaultSection = Menu->FindSection(NAME_None);
+	DefaultSection->Label = LOCTEXT("SettingsMenu_OutputLogSection_Label", "Output Log");
+}
+
+void FSwitchboardListenerMainWindow::CustomizeToolMenus_AddGeneralSection(UToolMenu* InMenu)
+{
+	FToolMenuSection& GeneralSection = InMenu->AddSection("SBL_General", LOCTEXT("SettingsMenu_GeneralSection_Label", "General"), FToolMenuInsert(NAME_None, EToolMenuInsertType::First));
+
+	GeneralSection.AddDynamicEntry("LaunchOnLoginDynamic", FNewToolMenuSectionDelegate::CreateLambda([this](FToolMenuSection& InSection)
+	{
+		const FString AutolaunchExe = UE::SwitchboardListener::Autolaunch::GetInvocationExecutable(LogSwitchboard);
+		const FString CurrentExe = FPlatformProcess::ExecutablePath();
+		if (AutolaunchExe.IsEmpty())
+		{
+			CachedAutolaunchEnabled = ECheckBoxState::Unchecked;
+		}
+		else if (AutolaunchExe == CurrentExe)
+		{
+			CachedAutolaunchEnabled = ECheckBoxState::Checked;
+		}
+		else
+		{
+			CachedAutolaunchEnabled = ECheckBoxState::Undetermined;
+		}
+
+		FToolUIAction LaunchOnLoginAction;
+		LaunchOnLoginAction.GetActionCheckState = FToolMenuGetActionCheckState::CreateLambda([this](const FToolMenuContext&)
+			{
+				return CachedAutolaunchEnabled;
+			});
+		LaunchOnLoginAction.ExecuteAction = FToolMenuExecuteAction::CreateLambda([this, CurrentExe](const FToolMenuContext&)
+			{
+				if (CachedAutolaunchEnabled == ECheckBoxState::Checked)
+				{
+					if (ensure(UE::SwitchboardListener::Autolaunch::RemoveInvocation(LogSwitchboard)))
+					{
+						CachedAutolaunchEnabled = ECheckBoxState::Unchecked;
+					}
+				}
+				else
+				{
+					if (ensure(UE::SwitchboardListener::Autolaunch::SetInvocation(FString::Printf(TEXT("\"%s\""), *CurrentExe), LogSwitchboard)))
+					{
+						CachedAutolaunchEnabled = ECheckBoxState::Checked;
+					}
+				}
+			});
+
+		InSection.AddMenuEntry(
+			"LaunchOnLogin",
+			LOCTEXT("SettingsMenu_LaunchOnLogin_Label", "Launch Switchboard Listener on Login"),
+			FText(),
+			FSlateIcon(),
+			LaunchOnLoginAction,
+			EUserInterfaceActionType::ToggleButton
+		);
+	}));
+}
+
+void FSwitchboardListenerMainWindow::CustomizeToolMenus_AddPasswordSection(UToolMenu* InMenu)
+{
+	FToolMenuSection& PasswordSection = InMenu->AddSection("SBL_Password", LOCTEXT("SettingsMenu_PasswordSection_Label", "Password"), FToolMenuInsert("SBL_General", EToolMenuInsertType::After));
+
+	PasswordSection.AddDynamicEntry("ShowPasswordDynamic", FNewToolMenuSectionDelegate::CreateLambda([this](FToolMenuSection& InSection)
+		{
+			// Reset to false every time menu is opened.
+			bPasswordRevealedInMenu = false;
+
+			FToolUIAction ShowPasswordAction;
+			ShowPasswordAction.GetActionCheckState = FToolMenuGetActionCheckState::CreateLambda([this](const FToolMenuContext&)
+				{
+					return bPasswordRevealedInMenu ? ECheckBoxState::Checked : ECheckBoxState::Unchecked;
+				});
+			ShowPasswordAction.ExecuteAction = FToolMenuExecuteAction::CreateLambda([this](const FToolMenuContext&)
+				{
+					bPasswordRevealedInMenu = !bPasswordRevealedInMenu;
+				});
+
+			InSection.AddMenuEntry("ShowPassword",
+				TAttribute<FText>::CreateLambda([this]()
+					{
+						return bPasswordRevealedInMenu
+							? FText::FromString(Listener.GetAuthPassword())
+							: LOCTEXT("SettingsMenu_ShowPassword_Placeholder", "••••••");
+					}),
+				FText(),
+				TAttribute<FSlateIcon>::CreateLambda([this]()
+				{
+					return bPasswordRevealedInMenu
+						? FSlateIcon(FAppStyle::Get().GetStyleSetName(), "Icons.Visible")
+						: FSlateIcon(FAppStyle::Get().GetStyleSetName(), "Icons.Hidden");
+				}),
+				ShowPasswordAction,
+				EUserInterfaceActionType::ToggleButton
+			).bShouldCloseWindowAfterMenuSelection = false;
+		}));
+
+	PasswordSection.AddMenuEntry("ResetPassword", LOCTEXT("SettingsMenu_ResetPassword_Label", "Reset Password"), FText(), FSlateIcon(),
+		FExecuteAction::CreateLambda([this]()
+			{
+				PasswordTextBox->SetText(FText::FromString(Listener.GetAuthPassword()));
+				PanelSwitcher->SetActiveWidgetIndex(static_cast<int32>(EPanelIndices::PasswordPanel));
+			})
+	);
+}
+
+void FSwitchboardListenerMainWindow::CustomizeToolMenus_AddDevelopmentSection(UToolMenu* InMenu)
+{
+	FToolMenuSection& DevelopmentSection = InMenu->AddSection("SBL_Development", LOCTEXT("SettingsMenu_DevelopmentSection_Label", "Development"), FToolMenuInsert(NAME_None, EToolMenuInsertType::After));
+	DevelopmentSection.AddMenuEntry("WidgetReflector", LOCTEXT("SettingsMenu_WidgetReflector_Label", "Widget Reflector"), FText(), FSlateIcon(),
+		FExecuteAction::CreateStatic([]()
+			{
+				FModuleManager::LoadModuleChecked<ISlateReflectorModule>("SlateReflector").DisplayWidgetReflector();
+			})
+	);
+}
+
+
+TSharedRef<SWidget> FSwitchboardListenerMainWindow::CreateRootSwitcher()
+{
+	ensure(!PanelSwitcher.IsValid());
+	PanelSwitcher = SNew(SWidgetSwitcher);
+
+	PanelSwitcher->AddSlot(static_cast<int32>(EPanelIndices::MainPanel))
+		[
+			SNew(SOverlay)
+			+ SOverlay::Slot()
+			[
+				SNew(SImage)
+				.Image(FAppStyle::GetBrush("Brushes.Panel"))
+				.Visibility(EVisibility::HitTestInvisible)
+			]
+			+ SOverlay::Slot()
+			[
+				CreateMainPanel()
+			]
+		];
+
+	PanelSwitcher->AddSlot(static_cast<int32>(EPanelIndices::PasswordPanel))
+		[
+			CreateSetPasswordPanel()
+		];
+
+	return PanelSwitcher.ToSharedRef();
+}
+
+
+TSharedRef<SWidget> FSwitchboardListenerMainWindow::CreateSetPasswordPanel()
+{
+	TSharedPtr<STextBlock> HeadingTextBlock, IntroTextBlock;
+
+	TSharedRef<SWidget> SetPasswordPanel = SNew(SHorizontalBox)
+		+ SHorizontalBox::Slot()
+		.FillWidth(1.0f)
+		[
+			SNullWidget::NullWidget
+		]
+		+ SHorizontalBox::Slot()
+		.AutoWidth()
+		.HAlign(HAlign_Center)
+		[
+			SNew(SVerticalBox)
+			+ SVerticalBox::Slot()
+			.FillHeight(1.0f)
+			[
+				SNullWidget::NullWidget
+			]
+			+ SVerticalBox::Slot()
+			.AutoHeight()
+			[
+				SAssignNew(HeadingTextBlock, STextBlock)
+				.Text(LOCTEXT("PasswordPanel_WelcomeHeading", "Welcome to Switchboard Listener"))
+				.Font(FAppStyle::Get().GetFontStyle("HeadingMedium"))
+				.TransformPolicy(ETextTransformPolicy::ToUpper)
+			]
+			+ SVerticalBox::Slot()
+			.AutoHeight()
+			.Padding(0.0f, 24.0f)
+			[
+				SAssignNew(IntroTextBlock, STextBlock)
+				.Text(LOCTEXT("PasswordPanel_WelcomeIntro", "Please set a password for Switchboard Listener on this machine. This password will be used to authenticate and establish a secure connection with Switchboard."))
+				.Font(FCoreStyle::GetDefaultFontStyle("Medium", 12))
+			]
+			+ SVerticalBox::Slot()
+			.AutoHeight()
+			.Padding(0.0f, 8.0f)
+			[
+				SNew(SHorizontalBox)
+				+ SHorizontalBox::Slot()
+				.FillWidth(1.0f)
+				[
+					SAssignNew(PasswordTextBox, SEditableTextBox)
+					.IsPassword_Lambda([this]() { return !bPasswordRevealedInEditBox; })
+					.HintText(LOCTEXT("PasswordPanel_Edit_Hint", "Password"))
+				]
+				+ SHorizontalBox::Slot()
+				.FillWidth(3.0f)
+				[
+					SNew(SSpacer)
+				]
+			]
+			+ SVerticalBox::Slot()
+			.AutoHeight()
+			.Padding(0.0f, 8.0f)
+			[
+				SNew(SHorizontalBox)
+				+ SHorizontalBox::Slot()
+				.AutoWidth()
+				[
+					SNew(SCheckBox)
+					.IsChecked_Lambda([this]() { return bPasswordRevealedInEditBox ? ECheckBoxState::Checked : ECheckBoxState::Unchecked; })
+					.OnCheckStateChanged_Lambda([this](ECheckBoxState InNewState)
+						{
+							bPasswordRevealedInEditBox = (InNewState == ECheckBoxState::Checked);
+						})
+				]
+				+ SHorizontalBox::Slot()
+				.VAlign(VAlign_Center)
+				.Padding(4.0f, 0.0f)
+				[
+					SNew(STextBlock)
+					.Text(LOCTEXT("PasswordPanel_ShowPasswordCheckLabel", "Show Password"))
+				]
+			]
+			+ SVerticalBox::Slot()
+			.AutoHeight()
+			.Padding(0.0f, 8.0f)
+			[
+				SNew(SHorizontalBox)
+				+ SHorizontalBox::Slot()
+				.FillWidth(1.0f)
+				[
+					SNew(SButton)
+					.Text_Lambda([this]()
+						{
+							return Listener.GetAuthPassword().IsEmpty()
+								? LOCTEXT("PasswordPanel_SetPasswordButton", "Set Password")
+								: LOCTEXT("PasswordPanel_ChangePasswordButton", "Change Password");
+						})
+					.IsEnabled_Lambda([this]() { return !PasswordTextBox->GetText().IsEmpty(); })
+					.OnClicked_Raw(this, &FSwitchboardListenerMainWindow::OnSetPasswordClicked)
+					.ButtonStyle(&FAppStyle::Get(), "PrimaryButton")
+					.TextStyle(&FAppStyle::Get(), "PrimaryButtonText")
+					.ContentPadding(FMargin(4.0, 2.0 + 2.0))
+					.HAlign(HAlign_Center)
+				]
+				+ SHorizontalBox::Slot()
+				.FillWidth(1.0f)
+				.Padding(4.0f, 0.0f)
+				[
+					SNew(SButton)
+					.Text(LOCTEXT("PasswordPanel_CancelButton", "Cancel"))
+					.Visibility_Lambda([this]() { return Listener.GetAuthPassword().IsEmpty() ? EVisibility::Collapsed : EVisibility::Visible; })
+					.OnClicked_Lambda([this]()
+						{
+							PanelSwitcher->SetActiveWidgetIndex(static_cast<int32>(EPanelIndices::MainPanel));
+							return FReply::Handled();
+						})
+					.ContentPadding(FMargin(4.0, 2.0 + 2.0))
+					.HAlign(HAlign_Center)
+				]
+				+ SHorizontalBox::Slot()
+				.FillWidth(2.0f)
+				[
+					SNew(SSpacer)
+				]
+			]
+			+ SVerticalBox::Slot()
+			.FillHeight(1.0f)
+			[
+				SNew(SSpacer)
+			]
+		]
+		+ SHorizontalBox::Slot()
+		.FillWidth(1.0f)
+		[
+			SNew(SSpacer)
+		];
+
+	// Constrain the body paragraph to the width of the heading.
+	const float LayoutScaling = FSlateApplication::Get().GetApplicationScale() * RootWindow->GetNativeWindow()->GetDPIScaleFactor();
+	const float HeadingWidth = HeadingTextBlock->ComputeDesiredSize(LayoutScaling).X;
+	IntroTextBlock->SetWrapTextAt(HeadingWidth);
+
+	return SetPasswordPanel;
+}
+
+
+TSharedRef<SWidget> FSwitchboardListenerMainWindow::CreateMainPanel()
+{
+	return SNew(SVerticalBox)
+		+ SVerticalBox::Slot()
+		.FillHeight(1.0)
+		[
+			CreateOutputLog()
+		]
+		+ SVerticalBox::Slot()
+		.AutoHeight()
+		.Padding(8.0f, 4.0f, 0.0f, 4.0f)
+		[
+			SNew(SHorizontalBox)
+			+ SHorizontalBox::Slot()
+			.AutoWidth()
+			.VAlign(VAlign_Center)
+			[
+				SNew(STextBlock)
+				.Text(LOCTEXT("ConnectionStatusIndicator", "Status"))
+				.Font(FCoreStyle::GetDefaultFontStyle("Medium", 12))
+				.ColorAndOpacity_Lambda([this]()
+					{
+						switch (Listener.GetConnectedClientAddresses().Num())
+						{
+							case 0: return FStyleColors::AccentRed;
+							case 1: return FStyleColors::AccentGreen;
+							default: return FStyleColors::AccentOrange;
+						}
+					})
+			]
+			+ SHorizontalBox::Slot()
+			.AutoWidth()
+			.Padding(8.0f, 0.0f, 0.0f, 0.0f)
+			.VAlign(VAlign_Center)
+			[
+				SNew(STextBlock)
+				.Text_Lambda([this]()
+					{
+						const TArray<FIPv4Address> ConnectedAddrs = Listener.GetConnectedClientAddresses().Array();
+						if (ConnectedAddrs.Num() <= 0)
+						{
+							return LOCTEXT("ConnectionStatusText_NotConnected", "Not Connected");
+						}
+						else if (ConnectedAddrs.Num() == 1)
+						{
+							return FText::Format(
+								LOCTEXT("ConnectionStatusText_Connected",
+								        "Connected to {0}"),
+								FText::FromString(ConnectedAddrs[0].ToString()));
+						}
+						else
+						{
+							return LOCTEXT("ConnectionStatusText_ConnectedMulti", "Connected to multiple clients");
+						}
+					})
+			]
+			+ SHorizontalBox::Slot()
+			.FillWidth(1.0)
+			[
+				SNew(SSpacer)
+			]
+		]
+	;
+}
+
+
+TSharedRef<SWidget> FSwitchboardListenerMainWindow::CreateOutputLog()
+{
+	FOutputLogModule& OutputLogModule = FModuleManager::Get().LoadModuleChecked<FOutputLogModule>("OutputLog");
+
+	FOutputLogCreationParams OutputLogCreationParams;
+
+	OutputLogCreationParams.AllowAsInitialLogCategory = FAllowLogCategoryCallback::CreateLambda([](const FName LogCategory) {
+		return LogCategory == LogSwitchboard.GetCategoryName();
+	});
+
+	return OutputLogModule.MakeOutputLogWidget(OutputLogCreationParams);
+}
+
+
+FReply FSwitchboardListenerMainWindow::OnSetPasswordClicked()
+{
+	const FString EnteredPassword = PasswordTextBox->GetText().ToString();
+	if (ensure(!EnteredPassword.IsEmpty()))
+	{
+		Listener.SetAuthPassword(EnteredPassword);
+		if (!Listener.IsListening())
+		{
+			if (ensure(Listener.StartListening()))
+			{
+				UE_LOGFMT(LogSwitchboard, Display, "Started listening");
+			}
+			else
+			{
+				UE_LOGFMT(LogSwitchboard, Error, "StartListening failed");
+			}
+		}
+
+		PanelSwitcher->SetActiveWidgetIndex(static_cast<int32>(EPanelIndices::MainPanel));
+	}
+
+	return FReply::Handled();
+}
+
+
+#undef LOCTEXT_NAMESPACE
diff --git a/Engine/Source/Programs/SwitchboardListener/SblSlate/SblSlate.Build.cs b/Engine/Source/Programs/SwitchboardListener/SblSlate/SblSlate.Build.cs
new file mode 100644
--- /dev/null
+++ b/Engine/Source/Programs/SwitchboardListener/SblSlate/SblSlate.Build.cs
@@ -0,0 +1,31 @@
+// Copyright Epic Games, Inc. All Rights Reserved.
+
+using UnrealBuildTool;
+
+
+public class SblSlate : ModuleRules
+{
+	public SblSlate(ReadOnlyTargetRules Target) : base(Target)
+	{
+		PublicDefinitions.Add("SBL_SLATE=1");
+
+		PublicDependencyModuleNames.AddRange(
+			new string[]
+			{
+				"SblCore",
+				"SlateCore",
+			}
+		);
+
+		PrivateDependencyModuleNames.AddRange(
+			new string[]
+			{
+				"OutputLog",
+				"Slate",
+				"SlateReflector",
+				"StandaloneRenderer",
+				"ToolMenus",
+			}
+		);
+	}
+}
diff --git a/Engine/Source/Programs/SwitchboardListener/SwitchboardListener.Build.cs b/Engine/Source/Programs/SwitchboardListener/SwitchboardListener.Build.cs
deleted file mode 100644
diff --git a/Engine/Source/Programs/SwitchboardListener/SwitchboardListener.Target.cs b/Engine/Source/Programs/SwitchboardListener/SwitchboardListener.Target.cs
--- a/Engine/Source/Programs/SwitchboardListener/SwitchboardListener.Target.cs
+++ b/Engine/Source/Programs/SwitchboardListener/SwitchboardListener.Target.cs
@@ -1,15 +1,16 @@
 // Copyright Epic Games, Inc. All Rights Reserved.
 
 using UnrealBuildTool;
-using System.Collections.Generic;
 
-public class SwitchboardListenerTarget : TargetRules
+
+public class SwitchboardListenerTargetBase : TargetRules
 {
-	public SwitchboardListenerTarget(TargetInfo Target) : base(Target)
+	public SwitchboardListenerTargetBase(TargetInfo Target) : base(Target)
 	{
+		SolutionDirectory = "Programs/Switchboard";
+
 		Type = TargetType.Program;
 		LinkType = TargetLinkType.Monolithic;
-		LaunchModuleName = "SwitchboardListener";
 
 		IncludeOrderVersion = EngineIncludeOrderVersion.Latest;
 
@@ -21,9 +22,15 @@ public class SwitchboardListenerTarget : TargetRules
 		// Disable internationalization support, so the exe can be launched outside of the engine directory.
 		bCompileICU = false;
 
-		// The listener is meant to be a console application (no window), but on MacOS, to get a proper log console, a full application must be built.
-		bIsBuildingConsoleApplication = Target.Platform != UnrealTargetPlatform.Mac;
-
 		bEnableTrace = true;
 	}
 }
+
+
+public class SwitchboardListenerTarget : SwitchboardListenerTargetBase
+{
+	public SwitchboardListenerTarget(TargetInfo Target) : base(Target)
+	{
+		LaunchModuleName = "SwitchboardListener";
+	}
+}
diff --git a/Engine/Source/Programs/SwitchboardListener/SwitchboardListener/Private/SwitchboardListenerMain.cpp b/Engine/Source/Programs/SwitchboardListener/SwitchboardListener/Private/SwitchboardListenerMain.cpp
new file mode 100644
--- /dev/null
+++ b/Engine/Source/Programs/SwitchboardListener/SwitchboardListener/Private/SwitchboardListenerMain.cpp
@@ -0,0 +1,105 @@
+// Copyright Epic Games, Inc. All Rights Reserved.
+
+#include "SwitchboardListenerMainCommon.h"
+#include "SwitchboardListener.h"
+#include "SwitchboardListenerVersion.h"
+#include "SblMainWindow.h"
+
+#include "RequiredProgramMainCPPInclude.h"
+
+
+IMPLEMENT_APPLICATION(SwitchboardListener, "SwitchboardListener");
+
+
+int32 SwitchboardListenerMain()
+{
+	const TCHAR* CommandLine = FCommandLine::Get();
+	const FSwitchboardCommandLineOptions Options = FSwitchboardCommandLineOptions::FromString(CommandLine);
+
+	if (Options.OutputVersion)
+	{
+		printf("SwitchboardListener %u.%u.%u", SBLISTENER_VERSION_MAJOR, SBLISTENER_VERSION_MINOR, SBLISTENER_VERSION_PATCH);
+		RequestEngineExit(TEXT("Output -version"));
+		return 0;
+	}
+
+	if (!SwitchboardListenerMainInit(Options))
+	{
+		return 1;
+	}
+
+	FSwitchboardListener Listener(Options);
+	FSwitchboardListenerMainWindow MainWindow(Listener);
+
+	const bool bStartListeningImmediately = false;
+	if (!Listener.Init(bStartListeningImmediately))
+	{
+		RequestEngineExit(TEXT("FSwitchboardListener::Init() failure"));
+		return 1;
+	}
+
+	UE_LOG(LogSwitchboard, Display, TEXT("Initialized"));
+
+	double LastTime = FPlatformTime::Seconds();
+	const float IdealFrameTime = 1.0f / 30.0f;
+
+	while (!IsEngineExitRequested())
+	{
+		BeginExitIfRequested();
+
+		const double CurrentTime = FPlatformTime::Seconds();
+		const double DeltaTime = CurrentTime - LastTime;
+		LastTime = CurrentTime;
+
+		const double HitchThresholdSec = 0.5;
+		if (DeltaTime > HitchThresholdSec)
+		{
+			UE_LOG(LogSwitchboard, Warning, TEXT("Hitch detected; %.3f seconds since prior tick"), DeltaTime);
+		}
+
+		FTaskGraphInterface::Get().ProcessThreadUntilIdle(ENamedThreads::GameThread);
+
+		// Pump & Tick objects
+		FTSTicker::GetCoreTicker().Tick(DeltaTime);
+
+		Listener.Tick();
+
+		GFrameCounter++;
+		FStats::AdvanceFrame(false);
+		GLog->FlushThreadedLogs();
+
+		// Run garbage collection for the UObjects for the rest of the frame or at least to 2 ms
+		IncrementalPurgeGarbage(true, FMath::Max<float>(0.002f, IdealFrameTime - (FPlatformTime::Seconds() - CurrentTime)));
+
+		// Throttle main thread main fps by sleeping if we still have time
+		FPlatformProcess::Sleep(FMath::Max<float>(0.0f, IdealFrameTime - (FPlatformTime::Seconds() - CurrentTime)));
+	}
+
+	Listener.Shutdown();
+
+	SwitchboardListenerMainShutdown();
+
+	return 0;
+}
+
+
+#if PLATFORM_WINDOWS
+int WINAPI WinMain(_In_ HINSTANCE hInInstance, _In_opt_ HINSTANCE hPrevInstance, _In_ LPSTR lpCmdLine, _In_ int nShowCmd)
+{
+	hInstance = hInInstance;
+
+	const TCHAR* CommandLine = ::GetCommandLineW();
+	CommandLine = FCommandLine::RemoveExeName(CommandLine);
+	FCommandLine::Set(CommandLine);
+
+	return SwitchboardListenerMainWrapper();
+}
+#else
+INT32_MAIN_INT32_ARGC_TCHAR_ARGV()
+{
+	const FString CommandLine = FCommandLine::BuildFromArgV(nullptr, ArgC, ArgV, nullptr);
+	FCommandLine::Set(*CommandLine);
+
+	return SwitchboardListenerMainWrapper();
+}
+#endif
diff --git a/Engine/Source/Programs/SwitchboardListener/SwitchboardListener/SwitchboardListener.Build.cs b/Engine/Source/Programs/SwitchboardListener/SwitchboardListener/SwitchboardListener.Build.cs
new file mode 100644
--- /dev/null
+++ b/Engine/Source/Programs/SwitchboardListener/SwitchboardListener/SwitchboardListener.Build.cs
@@ -0,0 +1,16 @@
+// Copyright Epic Games, Inc. All Rights Reserved.
+
+using UnrealBuildTool;
+
+
+public class SwitchboardListener : SwitchboardListenerCmd
+{
+	public SwitchboardListener(ReadOnlyTargetRules Target) : base(Target)
+	{
+		PrivateDependencyModuleNames.AddRange(
+			new string[] {
+				"SblSlate",
+			}
+		);
+	}
+}
diff --git a/Engine/Source/Programs/SwitchboardListener/SwitchboardListenerCmd.Target.cs b/Engine/Source/Programs/SwitchboardListener/SwitchboardListenerCmd.Target.cs
new file mode 100644
--- /dev/null
+++ b/Engine/Source/Programs/SwitchboardListener/SwitchboardListenerCmd.Target.cs
@@ -0,0 +1,16 @@
+// Copyright Epic Games, Inc. All Rights Reserved.
+
+using UnrealBuildTool;
+
+
+public class SwitchboardListenerCmdTarget : SwitchboardListenerTargetBase
+{
+	public SwitchboardListenerCmdTarget(TargetInfo Target) : base(Target)
+	{
+		LaunchModuleName = "SwitchboardListenerCmd";
+
+		// The listener is meant to be a console application (no window), but
+		// on MacOS, to get a proper log console, a full application must be built.
+		bIsBuildingConsoleApplication = Target.Platform != UnrealTargetPlatform.Mac;
+	}
+}
diff --git a/Engine/Source/Programs/SwitchboardListener/SwitchboardListenerCmd/Private/SwitchboardListenerCmdMain.cpp b/Engine/Source/Programs/SwitchboardListener/SwitchboardListenerCmd/Private/SwitchboardListenerCmdMain.cpp
new file mode 100644
--- /dev/null
+++ b/Engine/Source/Programs/SwitchboardListener/SwitchboardListenerCmd/Private/SwitchboardListenerCmdMain.cpp
@@ -0,0 +1,128 @@
+// Copyright Epic Games, Inc. All Rights Reserved.
+
+#include "SwitchboardListenerMainCommon.h"
+#include "SwitchboardListener.h"
+#include "SwitchboardListenerVersion.h"
+
+#include "RequiredProgramMainCPPInclude.h"
+
+
+IMPLEMENT_APPLICATION(SwitchboardListenerCmd, "SwitchboardListenerCmd");
+
+
+int32 SwitchboardListenerMain()
+{
+	const TCHAR* CommandLine = FCommandLine::Get();
+	const FSwitchboardCommandLineOptions Options = FSwitchboardCommandLineOptions::FromString(CommandLine);
+
+	if (Options.OutputVersion)
+	{
+		printf("SwitchboardListener %u.%u.%u", SBLISTENER_VERSION_MAJOR, SBLISTENER_VERSION_MINOR, SBLISTENER_VERSION_PATCH);
+		RequestEngineExit(TEXT("Output -version"));
+		return 0;
+	}
+
+#if PLATFORM_WINDOWS
+	// If we were launched detached, detect and spawn a new console.
+	// TODO: Command line switch to suppress this behavior? Or explicitly opt-in for redeploy?
+	{
+		bool bIsStdoutAttachedToConsole = false;
+		HANDLE StdoutHandle = GetStdHandle(STD_OUTPUT_HANDLE);
+		if (StdoutHandle != INVALID_HANDLE_VALUE)
+		{
+			if (GetFileType(StdoutHandle) == FILE_TYPE_CHAR)
+			{
+				bIsStdoutAttachedToConsole = true;
+			}
+		}
+
+		if (!bIsStdoutAttachedToConsole)
+		{
+			AllocConsole();
+
+			FILE* OutIgnoredNewStreamPtr = nullptr;
+			freopen_s(&OutIgnoredNewStreamPtr, "CONIN$", "r", stdin);
+			freopen_s(&OutIgnoredNewStreamPtr, "CONOUT$", "w", stdout);
+			freopen_s(&OutIgnoredNewStreamPtr, "CONOUT$", "w", stderr);
+		}
+	}
+
+	// Sets the window title. It informs the version of the listener and the UE version it is based on.
+	const FString Title = FString::Printf(TEXT("Switchboard Listener %d.%d.%d based on Unreal Engine %s"),
+		SBLISTENER_VERSION_MAJOR, SBLISTENER_VERSION_MINOR, SBLISTENER_VERSION_PATCH, *FEngineVersion::Current().ToString());
+
+	::SetConsoleTitle(*Title);
+#endif
+
+	if (!SwitchboardListenerMainInit(Options))
+	{
+		return 1;
+	}
+
+	FSwitchboardListener Listener(Options);
+
+	if (!Listener.Init())
+	{
+		RequestEngineExit(TEXT("FSwitchboardListener::Init() failure"));
+		return 1;
+	}
+
+#if PLATFORM_WINDOWS && !UE_BUILD_DEBUG
+	if (Options.MinimizeOnLaunch)
+	{
+		::ShowWindow(::GetConsoleWindow(), SW_MINIMIZE);
+	}
+#endif
+
+	UE_LOG(LogSwitchboard, Display, TEXT("Initialized and listening"));
+
+	double LastTime = FPlatformTime::Seconds();
+	const float IdealFrameTime = 1.0f / 30.0f;
+
+	while (!IsEngineExitRequested())
+	{
+		BeginExitIfRequested();
+
+		const double CurrentTime = FPlatformTime::Seconds();
+		const double DeltaTime = CurrentTime - LastTime;
+		LastTime = CurrentTime;
+
+		const double HitchThresholdSec = 0.5;
+		if (DeltaTime > HitchThresholdSec)
+		{
+			UE_LOG(LogSwitchboard, Warning, TEXT("Hitch detected; %.3f seconds since prior tick"), DeltaTime);
+		}
+
+		FTaskGraphInterface::Get().ProcessThreadUntilIdle(ENamedThreads::GameThread);
+
+		// Pump & Tick objects
+		FTSTicker::GetCoreTicker().Tick(DeltaTime);
+
+		Listener.Tick();
+
+		GFrameCounter++;
+		FStats::AdvanceFrame(false);
+		GLog->FlushThreadedLogs();
+
+		// Run garbage collection for the UObjects for the rest of the frame or at least to 2 ms
+		IncrementalPurgeGarbage(true, FMath::Max<float>(0.002f, IdealFrameTime - (FPlatformTime::Seconds() - CurrentTime)));
+
+		// Throttle main thread main fps by sleeping if we still have time
+		FPlatformProcess::Sleep(FMath::Max<float>(0.0f, IdealFrameTime - (FPlatformTime::Seconds() - CurrentTime)));
+	}
+
+	Listener.Shutdown();
+
+	SwitchboardListenerMainShutdown();
+
+	return 0;
+}
+
+
+INT32_MAIN_INT32_ARGC_TCHAR_ARGV()
+{
+	const FString CommandLine = FCommandLine::BuildFromArgV(nullptr, ArgC, ArgV, nullptr);
+	FCommandLine::Set(*CommandLine);
+
+	return SwitchboardListenerMainWrapper();
+}
diff --git a/Engine/Source/Programs/SwitchboardListener/SwitchboardListenerCmd/SwitchboardListenerCmd.Build.cs b/Engine/Source/Programs/SwitchboardListener/SwitchboardListenerCmd/SwitchboardListenerCmd.Build.cs
new file mode 100644
--- /dev/null
+++ b/Engine/Source/Programs/SwitchboardListener/SwitchboardListenerCmd/SwitchboardListenerCmd.Build.cs
@@ -0,0 +1,21 @@
+// Copyright Epic Games, Inc. All Rights Reserved.
+
+using UnrealBuildTool;
+
+public class SwitchboardListenerCmd : ModuleRules
+{
+	public SwitchboardListenerCmd(ReadOnlyTargetRules Target) : base(Target)
+	{
+		PrivateDependencyModuleNames.AddRange(
+			new string[] {
+				"SblCore",
+			}
+		);
+
+		PrivateIncludePathModuleNames.AddRange(
+			new string[] {
+				"Launch",
+			}
+		);
+	}
+}
diff --git a/Engine/Source/Programs/SwitchboardListenerHelper/SwitchboardListenerHelper.Target.cs b/Engine/Source/Programs/SwitchboardListenerHelper/SwitchboardListenerHelper.Target.cs
--- a/Engine/Source/Programs/SwitchboardListenerHelper/SwitchboardListenerHelper.Target.cs
+++ b/Engine/Source/Programs/SwitchboardListenerHelper/SwitchboardListenerHelper.Target.cs
@@ -12,6 +12,7 @@ public class SwitchboardListenerHelperTarget : TargetRules
 		Type = TargetType.Program;
 		LinkType = TargetLinkType.Monolithic;
 		LaunchModuleName = "SwitchboardListenerHelper";
+		SolutionDirectory = "Programs/Switchboard";
 
 		IncludeOrderVersion = EngineIncludeOrderVersion.Latest;
 
