SwitchboardListener: Linux build fixes.

#rnx

diff --git a/Engine/Source/Programs/SwitchboardListener/SblSlate/Private/SblMainWindow.cpp b/Engine/Source/Programs/SwitchboardListener/SblSlate/Private/SblMainWindow.cpp
--- a/Engine/Source/Programs/SwitchboardListener/SblSlate/Private/SblMainWindow.cpp
+++ b/Engine/Source/Programs/SwitchboardListener/SblSlate/Private/SblMainWindow.cpp
@@ -58,7 +58,9 @@ public:
 
 FSwitchboardListenerMainWindow::FSwitchboardListenerMainWindow(FSwitchboardListener& InListener)
 	: Listener(InListener)
+#if SWITCHBOARD_LISTENER_AUTOLAUNCH
 	, CachedAutolaunchEnabled(ECheckBoxState::Unchecked)
+#endif
 {
 	Listener.OnInit().AddRaw(this, &FSwitchboardListenerMainWindow::OnInit);
 	Listener.OnShutdown().AddRaw(this, &FSwitchboardListenerMainWindow::OnShutdown);
@@ -150,6 +152,7 @@ void FSwitchboardListenerMainWindow::CustomizeToolMenus_AddGeneralSection(UToolM
 {
 	FToolMenuSection& GeneralSection = InMenu->AddSection("SBL_General", LOCTEXT("SettingsMenu_GeneralSection_Label", "General"), FToolMenuInsert(NAME_None, EToolMenuInsertType::First));
 
+#if SWITCHBOARD_LISTENER_AUTOLAUNCH
 	GeneralSection.AddDynamicEntry("LaunchOnLoginDynamic", FNewToolMenuSectionDelegate::CreateLambda([this](FToolMenuSection& InSection)
 	{
 		const FString AutolaunchExe = UE::SwitchboardListener::Autolaunch::GetInvocationExecutable(LogSwitchboard);
@@ -199,6 +202,7 @@ void FSwitchboardListenerMainWindow::CustomizeToolMenus_AddGeneralSection(UToolM
 			EUserInterfaceActionType::ToggleButton
 		);
 	}));
+#endif // #if SWITCHBOARD_LISTENER_AUTOLAUNCH
 }
 
 void FSwitchboardListenerMainWindow::CustomizeToolMenus_AddPasswordSection(UToolMenu* InMenu)
