SwitchboardListener: Slate tweaks.

- Hide Output Log category separator.
- Hide widget reflector in non-debug builds.
- Relabel "Open Source Location" as "Open Logs Folder."
- Remove "Show Password" menu item.
- Relabel "Reset Password" item to "Show/Change Password."
- Use Regular font weight for password panel body copy.
- Enumerate connected addresses in status tooltip when we have more than one incoming connection.

#rnx

diff --git a/Engine/Source/Programs/SwitchboardListener/SblSlate/Internal/SblMainWindow.h b/Engine/Source/Programs/SwitchboardListener/SblSlate/Internal/SblMainWindow.h
--- a/Engine/Source/Programs/SwitchboardListener/SblSlate/Internal/SblMainWindow.h
+++ b/Engine/Source/Programs/SwitchboardListener/SblSlate/Internal/SblMainWindow.h
@@ -36,7 +36,9 @@ private:
 	void CustomizeToolMenus();
 	void CustomizeToolMenus_AddGeneralSection(UToolMenu* InMenu);
 	void CustomizeToolMenus_AddPasswordSection(UToolMenu* InMenu);
+#if UE_BUILD_DEBUG
 	void CustomizeToolMenus_AddDevelopmentSection(UToolMenu* InMenu);
+#endif
 
 	TSharedRef<SWidget> CreateRootSwitcher();
 	TSharedRef<SWidget> CreateSetPasswordPanel();
diff --git a/Engine/Source/Programs/SwitchboardListener/SblSlate/Private/SblMainWindow.cpp b/Engine/Source/Programs/SwitchboardListener/SblSlate/Private/SblMainWindow.cpp
--- a/Engine/Source/Programs/SwitchboardListener/SblSlate/Private/SblMainWindow.cpp
+++ b/Engine/Source/Programs/SwitchboardListener/SblSlate/Private/SblMainWindow.cpp
@@ -140,14 +140,27 @@ void FSwitchboardListenerMainWindow::CustomizeToolMenus()
 	static const FName MenuName = "OutputLog.SettingsMenu";
 	UToolMenu* Menu = ToolMenus->ExtendMenu(MenuName);
 
+	FCustomizedToolMenu* CustomizedMenu = ToolMenus->AddRuntimeMenuCustomization(MenuName);
+	CustomizedMenu->AddEntry("Separator")->Visibility = ECustomizedToolMenuVisibility::Hidden;
+
 	CustomizeToolMenus_AddGeneralSection(Menu);
 	CustomizeToolMenus_AddPasswordSection(Menu);
+#if UE_BUILD_DEBUG
 	CustomizeToolMenus_AddDevelopmentSection(Menu);
+#endif
+
+	if (FToolMenuSection* DefaultSection = Menu->FindSection(NAME_None); ensure(DefaultSection))
+	{
+		DefaultSection->Label = LOCTEXT("SettingsMenu_OutputLogSection_Label", "Output Log");
 
-	FToolMenuSection* DefaultSection = Menu->FindSection(NAME_None);
-	DefaultSection->Label = LOCTEXT("SettingsMenu_OutputLogSection_Label", "Output Log");
+		if (FToolMenuEntry* BrowseLogsEntry = DefaultSection->FindEntry("BrowseLogDirectory"); ensure(BrowseLogsEntry))
+		{
+			BrowseLogsEntry->Label = LOCTEXT("SettingsMenu_BrowseLogDirectory_Label", "Open Logs Folder");
+		}
+	}
 }
 
+
 void FSwitchboardListenerMainWindow::CustomizeToolMenus_AddGeneralSection(UToolMenu* InMenu)
 {
 	FToolMenuSection& GeneralSection = InMenu->AddSection("SBL_General", LOCTEXT("SettingsMenu_GeneralSection_Label", "General"), FToolMenuInsert(NAME_None, EToolMenuInsertType::First));
@@ -205,45 +218,12 @@ void FSwitchboardListenerMainWindow::CustomizeToolMenus_AddGeneralSection(UToolM
 #endif // #if SWITCHBOARD_LISTENER_AUTOLAUNCH
 }
 
+
 void FSwitchboardListenerMainWindow::CustomizeToolMenus_AddPasswordSection(UToolMenu* InMenu)
 {
 	FToolMenuSection& PasswordSection = InMenu->AddSection("SBL_Password", LOCTEXT("SettingsMenu_PasswordSection_Label", "Password"), FToolMenuInsert("SBL_General", EToolMenuInsertType::After));
 
-	PasswordSection.AddDynamicEntry("ShowPasswordDynamic", FNewToolMenuSectionDelegate::CreateLambda([this](FToolMenuSection& InSection)
-		{
-			// Reset to false every time menu is opened.
-			bPasswordRevealedInMenu = false;
-
-			FToolUIAction ShowPasswordAction;
-			ShowPasswordAction.GetActionCheckState = FToolMenuGetActionCheckState::CreateLambda([this](const FToolMenuContext&)
-				{
-					return bPasswordRevealedInMenu ? ECheckBoxState::Checked : ECheckBoxState::Unchecked;
-				});
-			ShowPasswordAction.ExecuteAction = FToolMenuExecuteAction::CreateLambda([this](const FToolMenuContext&)
-				{
-					bPasswordRevealedInMenu = !bPasswordRevealedInMenu;
-				});
-
-			InSection.AddMenuEntry("ShowPassword",
-				TAttribute<FText>::CreateLambda([this]()
-					{
-						return bPasswordRevealedInMenu
-							? FText::FromString(Listener.GetAuthPassword())
-							: LOCTEXT("SettingsMenu_ShowPassword_Placeholder", "••••••");
-					}),
-				FText(),
-				TAttribute<FSlateIcon>::CreateLambda([this]()
-				{
-					return bPasswordRevealedInMenu
-						? FSlateIcon(FAppStyle::Get().GetStyleSetName(), "Icons.Visible")
-						: FSlateIcon(FAppStyle::Get().GetStyleSetName(), "Icons.Hidden");
-				}),
-				ShowPasswordAction,
-				EUserInterfaceActionType::ToggleButton
-			).bShouldCloseWindowAfterMenuSelection = false;
-		}));
-
-	PasswordSection.AddMenuEntry("ResetPassword", LOCTEXT("SettingsMenu_ResetPassword_Label", "Reset Password"), FText(), FSlateIcon(),
+	PasswordSection.AddMenuEntry("PasswordPanel", LOCTEXT("SettingsMenu_PasswordPanel_Label", "Show/Change Password"), FText(), FSlateIcon(),
 		FExecuteAction::CreateLambda([this]()
 			{
 				PasswordTextBox->SetText(FText::FromString(Listener.GetAuthPassword()));
@@ -252,6 +232,8 @@ void FSwitchboardListenerMainWindow::CustomizeToolMenus_AddPasswordSection(UTool
 	);
 }
 
+
+#if UE_BUILD_DEBUG
 void FSwitchboardListenerMainWindow::CustomizeToolMenus_AddDevelopmentSection(UToolMenu* InMenu)
 {
 	FToolMenuSection& DevelopmentSection = InMenu->AddSection("SBL_Development", LOCTEXT("SettingsMenu_DevelopmentSection_Label", "Development"), FToolMenuInsert(NAME_None, EToolMenuInsertType::After));
@@ -262,6 +244,7 @@ void FSwitchboardListenerMainWindow::CustomizeToolMenus_AddDevelopmentSection(UT
 			})
 	);
 }
+#endif // #if UE_BUILD_DEBUG
 
 
 TSharedRef<SWidget> FSwitchboardListenerMainWindow::CreateRootSwitcher()
@@ -327,7 +310,7 @@ TSharedRef<SWidget> FSwitchboardListenerMainWindow::CreateSetPasswordPanel()
 			[
 				SAssignNew(IntroTextBlock, STextBlock)
 				.Text(LOCTEXT("PasswordPanel_WelcomeIntro", "Please set a password for Switchboard Listener on this machine. This password will be used to authenticate and establish a secure connection with Switchboard."))
-				.Font(FCoreStyle::GetDefaultFontStyle("Medium", 12))
+				.Font(FCoreStyle::GetDefaultFontStyle("Regular", 12))
 			]
 			+ SVerticalBox::Slot()
 			.AutoHeight()
@@ -489,6 +472,24 @@ TSharedRef<SWidget> FSwitchboardListenerMainWindow::CreateMainPanel()
 							return LOCTEXT("ConnectionStatusText_ConnectedMulti", "Connected to multiple clients");
 						}
 					})
+				.ToolTipText_Lambda([this]()
+					{
+						const TArray<FIPv4Address> ConnectedAddrs = Listener.GetConnectedClientAddresses().Array();
+						if (ConnectedAddrs.Num() > 1)
+						{
+							FFormatOrderedArguments TextAddrs;
+							TextAddrs.Reserve(ConnectedAddrs.Num());
+							for (const FIPv4Address& Addr : ConnectedAddrs)
+							{
+								TextAddrs.Emplace(Addr.ToText());
+							}
+							const FText ConnectedAddrListText = FText::Join(LOCTEXT("ConnectionStatusTooltip_AddrDelimiter", ", "), TextAddrs);
+							return FText::Format(LOCTEXT("ConnectionStatusTooltip_ConnectedMulti", "Connected to multiple clients: {0}"),
+								ConnectedAddrListText);
+						}
+
+						return FText::GetEmpty();
+					})
 			]
 			+ SHorizontalBox::Slot()
 			.FillWidth(1.0)
