SwitchboardListener: Reintroduce engine version as part of window title, and include stream/build details within a tooltip, similar to the editor.

#rnx

diff --git a/Engine/Source/Programs/SwitchboardListener/SblSlate/Private/SblMainWindow.cpp b/Engine/Source/Programs/SwitchboardListener/SblSlate/Private/SblMainWindow.cpp
--- a/Engine/Source/Programs/SwitchboardListener/SblSlate/Private/SblMainWindow.cpp
+++ b/Engine/Source/Programs/SwitchboardListener/SblSlate/Private/SblMainWindow.cpp
@@ -3,6 +3,7 @@
 #include "SblMainWindow.h"
 #include "SwitchboardListener.h"
 #include "SwitchboardListenerApp.h"
+#include "SwitchboardListenerVersion.h"
 
 #include "Framework/Application/SlateApplication.h"
 #include "Framework/Application/SWindowTitleBar.h"
@@ -12,7 +13,9 @@
 #include "HAL/PlatformApplicationMisc.h"
 #include "ISlateReflectorModule.h"
 #include "Logging/StructuredLog.h"
+#include "Misc/App.h"
 #include "Misc/ConfigCacheIni.h"
+#include "Misc/EngineVersion.h"
 #include "Modules/ModuleManager.h"
 #include "OutputLogCreationParams.h"
 #include "OutputLogModule.h"
@@ -40,8 +43,13 @@ public:
 		FSlateApplication::Get().GetDisplayMetrics(DisplayMetrics);
 		const float DPIScaleFactor = FPlatformApplicationMisc::GetDPIScaleFactorAtPoint(DisplayMetrics.PrimaryDisplayWorkAreaRect.Left, DisplayMetrics.PrimaryDisplayWorkAreaRect.Top);
 
+		const FString EngineVersionTitleString = FEngineVersion::Current().ToString(EVersionComponent::Patch);
+		const FText WindowTitle = FText::Format(
+			LOCTEXT("WindowTitle", "Switchboard Listener (Unreal Engine {0})"),
+			FText::FromString(EngineVersionTitleString));
+
 		SWindow::Construct(SWindow::FArguments()
-			.Title(LOCTEXT("WindowTitle", "Switchboard Listener"))
+			.Title(WindowTitle)
 			.CreateTitleBar(true)
 			.SupportsMaximize(true)
 			.SupportsMinimize(true)
@@ -52,6 +60,32 @@ public:
 			.ClientSize(FVector2D(1000.0f * DPIScaleFactor, 500.0f * DPIScaleFactor))
 			.AdjustInitialSizeAndPositionForDPIScale(false)
 		);
+
+		// Create the tooltip showing more detailed information
+		const FString AppVersionString = FString::Printf(TEXT("%u.%u.%u"), SBLISTENER_VERSION_MAJOR, SBLISTENER_VERSION_MINOR, SBLISTENER_VERSION_PATCH);
+		const FString EngineVersionTooltipString = FEngineVersion::Current().ToString(FEngineVersion::Current().HasChangelist() ? EVersionComponent::Changelist : EVersionComponent::Patch);
+		const EBuildConfiguration BuildConfig = FApp::GetBuildConfiguration();
+
+		FFormatNamedArguments TooltipArgs;
+		TooltipArgs.Add(TEXT("AppVersion"), FText::FromString(AppVersionString));
+		TooltipArgs.Add(TEXT("EngineVersion"), FText::FromString(EngineVersionTooltipString));
+		TooltipArgs.Add(TEXT("Branch"), FText::FromString(FEngineVersion::Current().GetBranch()));
+		TooltipArgs.Add(TEXT("BuildConfiguration"), EBuildConfigurations::ToText(BuildConfig));
+		TooltipArgs.Add(TEXT("BuildDate"), FText::FromString(FApp::GetBuildDate()));
+
+		const FText TitleToolTip = FText::Format(
+			LOCTEXT("TitleBarTooltip", "Switchboard Listener Version: {AppVersion}\nEngine Version: {EngineVersion}\nBranch: {Branch}\nBuild Configuration: {BuildConfiguration}\nBuild Date: {BuildDate}"), TooltipArgs);
+
+		AddOverlaySlot()
+		.VAlign(VAlign_Top)
+		.HAlign(HAlign_Center)
+		[
+			SNew(STextBlock)
+			.Margin(FMargin(0.0, 12.0, 0.0, 0.0))
+			.Text(WindowTitle)
+			.ToolTipText(TitleToolTip)
+			.ColorAndOpacity(FLinearColor::Transparent)
+		];
 	}
 };
 
