CIS fix.

diff --git a/Engine/Source/Programs/SwitchboardListener/SblCore/Private/SwitchboardListener.cpp b/Engine/Source/Programs/SwitchboardListener/SblCore/Private/SwitchboardListener.cpp
--- a/Engine/Source/Programs/SwitchboardListener/SblCore/Private/SwitchboardListener.cpp
+++ b/Engine/Source/Programs/SwitchboardListener/SblCore/Private/SwitchboardListener.cpp
@@ -1370,7 +1370,7 @@ bool FSwitchboardListener::Task_ReceiveFileFromClient(const FSwitchboardReceiveF
 
 		if (TempDir.EndsWith(TEXT("/")) || TempDir.EndsWith(TEXT("\\")))
 		{
-			TempDir.LeftChopInline(1, false);
+			TempDir.LeftChopInline(1, EAllowShrinking::No);
 		}
 
 		Destination.ReplaceInline(TEXT("%TEMP%"), *TempDir);
