Switchboard: Implement password hashing, login tokens, PosixOS credential persistence.

- Authentication attempts now use scrypt hashing and constant-time comparison.
- Successful authentication replies include a JSON Web Token, which can be used for subsequent connections.
- On PosixOS family platforms, where we can't rely on the OS to encrypt credentials at rest, we do not offer the flow where users can view the current Switchboard Listener password; only the hash is stored. You can change the password, but not view the current password.
- Added automated tests.

Also fixes a race condition where FThreadStats could cause a crash on exit.

diff --git a/Engine/Plugins/VirtualProduction/Switchboard/Source/Switchboard/switchboard/credential_store.py b/Engine/Plugins/VirtualProduction/Switchboard/Source/Switchboard/switchboard/credential_store.py
--- a/Engine/Plugins/VirtualProduction/Switchboard/Source/Switchboard/switchboard/credential_store.py
+++ b/Engine/Plugins/VirtualProduction/Switchboard/Source/Switchboard/switchboard/credential_store.py
@@ -4,7 +4,13 @@ from __future__ import annotations
 
 from abc import abstractmethod, ABC
 from dataclasses import dataclass
+import json
+import logging
+import os
+import pathlib
+import stat
 import sys
+import tempfile
 import threading
 from typing import Callable, Optional
 
@@ -29,6 +35,10 @@ class CredentialStore(ABC, QtCore.QObject, metaclass=QObjectABCMeta):
     def __init__(self):
         super().__init__()
 
+    @classmethod
+    def encrypted_at_rest(cls) -> bool:
+        return False
+
     @abstractmethod
     def get(
         self,
@@ -50,10 +60,10 @@ class CredentialStore(ABC, QtCore.QObject, metaclass=QObjectABCMeta):
 
     @staticmethod
     def create() -> CredentialStore:
-        # TODO:
-        #  - Linux -> Kernel key retention service? (c.f. `keyctl`)
         if sys.platform.startswith('win'):
             return CredentialStoreWindows()
+        elif os.name == 'posix':
+            return CredentialStorePosix()
 
         return CredentialStoreGeneric()
 
@@ -131,6 +141,63 @@ class CredentialStoreGeneric(CredentialStore):
         self._queued_prompts.clear()
 
 
+class CredentialStorePosix(CredentialStoreGeneric):
+    PERMISSION_MASK = stat.S_IRWXU | stat.S_IRWXG | stat.S_IRWXO
+    CRED_FILE_FLAGS = stat.S_IRUSR | stat.S_IWUSR  # 600
+    CRED_DIR_FLAGS = stat.S_IRUSR | stat.S_IWUSR | stat.S_IXUSR  # 700
+
+    def __init__(self):
+        super().__init__()
+
+        del self._tempstore
+
+    @classmethod
+    def _credential_dir(cls) -> pathlib.Path:
+        return pathlib.Path(pathlib.Path.home(), '.switchboard')
+
+    @classmethod
+    def _credential_path(cls, key: str) -> pathlib.Path:
+        return pathlib.Path(cls._credential_dir(), f'credential_{key}.json')
+
+    def _try_get(self, key: str) -> Optional[CredentialStore.Credential]:
+        path = self._credential_path(key)
+        if not path or not path.is_file():
+            return None
+
+        stat = os.lstat(path)
+        current_uid = os.getuid()
+        if stat.st_uid != current_uid:
+            logging.error(
+                f'Credential file {path} is not owned by the current user!')
+
+        share_flags = stat.st_mode & self.PERMISSION_MASK
+        if share_flags != self.CRED_FILE_FLAGS:
+            logging.error(
+                f'Credential file {path} has incorrect permissions!')
+
+        with path.open() as f:
+            cred = json.load(f)
+            if ('user' in cred) and ('blob' in cred):
+                return CredentialStore.Credential(cred['user'], cred['blob'])
+            else:
+                logging.error(f'Error parsing credential: {cred}')
+                return None
+
+    def set(self, key: str, val: Optional[CredentialStore.Credential]) -> None:
+        cred_path = self._credential_path(key)
+        if val is None:
+            cred_path.unlink()
+            return
+
+        cred_dir = cred_path.parent
+        cred_dir.mkdir(mode=0o700, parents=True, exist_ok=True)
+
+        (tmp_fd, tmp_path) = tempfile.mkstemp(dir=cred_dir, text=True)
+        with os.fdopen(tmp_fd, mode='w') as f:
+            f.write(json.dumps({ 'user': val.username, 'blob': val.blob }))
+        os.replace(tmp_path, cred_path)
+
+
 if sys.platform.startswith('win'):
     import ctypes
     import ctypes.wintypes
@@ -242,8 +309,11 @@ if sys.platform.startswith('win'):
             self._CredReadW.errcheck = Cred_errcheck
             self._CredWriteW.errcheck = Cred_errcheck
 
+        @classmethod
+        def encrypted_at_rest(cls) -> bool:
+            return True
+
         def _try_get(self, key: str) -> Optional[CredentialStore.Credential]:
-            # return self._tempstore.get(key)
             need_credfree = False
             cred = ctypes.pointer(CREDENTIALW())
             try:
diff --git a/Engine/Plugins/VirtualProduction/Switchboard/Source/Switchboard/switchboard/listener_client.py b/Engine/Plugins/VirtualProduction/Switchboard/Source/Switchboard/switchboard/listener_client.py
--- a/Engine/Plugins/VirtualProduction/Switchboard/Source/Switchboard/switchboard/listener_client.py
+++ b/Engine/Plugins/VirtualProduction/Switchboard/Source/Switchboard/switchboard/listener_client.py
@@ -88,8 +88,15 @@ class SwitchboardClientProtocol(QuicConnectionProtocol):
                 return
 
             if self.client:
+                jwt: Optional[str] = None
+                password: Optional[str] = None
+                if not unsaved and not CREDENTIAL_STORE.encrypted_at_rest():
+                    jwt = cred.blob
+                else:
+                    password = cred.blob
+
                 msgid, msg = message_protocol.create_authenticate_message(
-                    cred.blob)
+                    jwt=jwt, password=password)
 
                 self.client.get_message_response(msgid, msg,
                                                  future=auth_response)
@@ -108,8 +115,13 @@ class SwitchboardClientProtocol(QuicConnectionProtocol):
             auth_succeeded = response['bAuthenticated'] == 'true'
 
             if auth_succeeded:
-                if closure_cred_unsaved and closure_credential:
-                    CREDENTIAL_STORE.set(host, closure_credential)
+                if CREDENTIAL_STORE.encrypted_at_rest():
+                    if closure_cred_unsaved and closure_credential:
+                        CREDENTIAL_STORE.set(host, closure_credential)
+                elif 'jwt' in response:
+                    credential = CredentialStore.Credential(
+                        fingerprint_str, response['jwt'])
+                    CREDENTIAL_STORE.set(host, credential)
             else:
                 # Clear invalid saved credential
                 if not closure_cred_unsaved:
diff --git a/Engine/Plugins/VirtualProduction/Switchboard/Source/Switchboard/switchboard/message_protocol.py b/Engine/Plugins/VirtualProduction/Switchboard/Source/Switchboard/switchboard/message_protocol.py
--- a/Engine/Plugins/VirtualProduction/Switchboard/Source/Switchboard/switchboard/message_protocol.py
+++ b/Engine/Plugins/VirtualProduction/Switchboard/Source/Switchboard/switchboard/message_protocol.py
@@ -2,19 +2,36 @@
 
 import base64
 import json
+from typing import Optional
 import uuid
 
 
-def create_authenticate_message(token: str):
+def create_authenticate_message(
+    *,
+    jwt: Optional[str] = None,
+    password: Optional[str] = None,
+):
+    assert jwt or password
+
     cmd_id = uuid.uuid4()
     message = {
         'command': 'authenticate',
         'id': str(cmd_id),
-        'token': token,
     }
+
+    if jwt:
+        message['jwt'] = jwt
+
+    if password:
+        message['password'] = password
+
+        # TODO: for backward compatibility with pre-release 5.4, remove
+        message['token'] = password
+
     message_json = json.dumps(message).encode() + b'\x00'
     return (cmd_id, message_json)
 
+
 def create_start_process_message(
     prog_path: str,
     prog_args: str,
diff --git a/Engine/Source/Programs/SwitchboardListener/SblCore/Internal/Linux/LinuxSwitchboardCredential.h b/Engine/Source/Programs/SwitchboardListener/SblCore/Internal/Linux/LinuxSwitchboardCredential.h
new file mode 100644
--- /dev/null
+++ b/Engine/Source/Programs/SwitchboardListener/SblCore/Internal/Linux/LinuxSwitchboardCredential.h
@@ -0,0 +1,16 @@
+// Copyright Epic Games, Inc. All Rights Reserved.
+
+#pragma once
+
+#include "PosixOS/PosixOSSwitchboardCredential.h"
+
+
+namespace UE::SwitchboardListener
+{
+	inline ICredentialManager& ICredentialManager::GetPlatformCredentialManager()
+	{
+		using namespace UE::SwitchboardListener::Private;
+		static FPosixCredentialManager CredentialManager;
+		return CredentialManager;
+	}
+};
diff --git a/Engine/Source/Programs/SwitchboardListener/SblCore/Internal/PosixOS/PosixOSSwitchboardCredential.h b/Engine/Source/Programs/SwitchboardListener/SblCore/Internal/PosixOS/PosixOSSwitchboardCredential.h
new file mode 100644
--- /dev/null
+++ b/Engine/Source/Programs/SwitchboardListener/SblCore/Internal/PosixOS/PosixOSSwitchboardCredential.h
@@ -0,0 +1,16 @@
+// Copyright Epic Games, Inc. All Rights Reserved.
+
+#pragma once
+
+#include "SwitchboardCredentialInterface.h"
+
+
+namespace UE::SwitchboardListener::Private
+{
+	class FPosixCredentialManager : public ICredentialManager
+	{
+	public:
+		virtual TSharedPtr<ICredential> LoadCredential(FStringView InCredentialName) override;
+		virtual bool SaveCredential(FStringView InCredentialName, FStringView InUser, FStringView InBlob) override;
+	};
+};
diff --git a/Engine/Source/Programs/SwitchboardListener/SblCore/Internal/SwitchboardAuth.h b/Engine/Source/Programs/SwitchboardListener/SblCore/Internal/SwitchboardAuth.h
new file mode 100644
--- /dev/null
+++ b/Engine/Source/Programs/SwitchboardListener/SblCore/Internal/SwitchboardAuth.h
@@ -0,0 +1,160 @@
+// Copyright Epic Games, Inc. All Rights Reserved.
+
+#pragma once
+
+#include "Containers/Map.h"
+#include "Containers/StringView.h"
+#include "Containers/UnrealString.h"
+#include "Containers/Utf8String.h"
+#include "Misc/CoreMiscDefines.h" // UE_PUSH_MACRO
+#include "Misc/DateTime.h"
+#include "Misc/TVariant.h"
+
+UE_PUSH_MACRO("UI")
+// Workaround for "ossl_typ.h(144): error C2365: 'UI': redefinition; previous definition was 'namespace'"
+//                "ObjectMacros.h(872): note: see declaration of 'UI'"
+#define UI UI_ST
+#include <openssl/ossl_typ.h>
+UE_POP_MACRO("UI")
+
+
+class FJsonWebToken;
+
+namespace UE::SwitchboardListener
+{
+	class ICredentialManager;
+}
+
+
+namespace UE::SwitchboardListener
+{
+	/**
+	 * Get the default paths for self-signed key pairs,
+	 * or (on some platforms) persisted credentials.
+	 * 
+	 * - Windows: %LocalAppData%\SwitchboardListener
+	 * - PosixOS family: ~/.switchboard/listener
+	 */
+	FString GetSwitchboardDataDirectory();
+
+	/**
+	 * Disable console echo and prompt the user to enter a password (twice).
+	 * 
+	 * @param bAllowEmpty Whether to allow the user to leave the input empty.
+	 * 
+	 * @return The password the user entered.
+	 */
+	FString ReadPasswordFromStdin(bool bAllowEmpty = false);
+
+	/**
+	 * Completely overwrites the contents of the provided array with
+	 * cryptographically secure pseudo-random bytes.
+	 * 
+	 * @param InOutArray Set to the desired length before calling.
+	 */
+	void FillSecureRandom(TArrayView<uint8> InOutArray);
+};
+
+
+/**
+ * - Uses either the specified key pair, or an automatically generated
+ *   self-signed key pair. On systems where credentials can be stored
+ *   encrypted at rest, the resulting private key is protected with a
+ *   randomly generated password.
+ *
+ * - Can sign and verify JSON Web Tokens using the loaded key pair.
+ * 
+ * - Can generate, persist, and constant-time compare password hashes using
+ *   scrypt key derivation.
+ * 
+ * - Has helper methods for MsQuic interop.
+ */
+class FSwitchboardAuthHelper
+{
+public:
+	struct FSettings
+	{
+		TOptional<TVariant<FString, X509*>> Certificate;
+		TOptional<TVariant<FString, EVP_PKEY*>> PrivateKey;
+		TOptional<FUtf8StringView> PrivateKeyPassword;
+
+		UE::SwitchboardListener::ICredentialManager* CredentialManager = nullptr;
+
+		FSettings()
+		{
+			// Without this (or with `= default`), clang gives the following:
+			// `error: default member initializer for 'CredentialManager' needed within definition
+			// of enclosing class 'FSwitchboardAuthHelper' outside of member functions`
+		}
+	};
+
+	bool Initialize(const FSettings& InSettings = FSettings());
+	void Shutdown();
+
+
+	/** May be empty if the certificate was not loaded from file. */
+	const FUtf8String& GetCertificateFilePath() const;
+
+	/** May be empty if the private key was not loaded from file. */
+	const FUtf8String& GetPrivateKeyFilePath() const;
+
+	/** May be null if the private key was not password-protected. */
+	const UTF8CHAR* GetPrivateKeyPassword() const;
+
+
+	/** Change the required password. */
+	bool SetAuthPassword(FStringView InNewPassword, bool bExpireExistingTokens = true);
+
+	/** Validate whether the provided password attempt matches the expected password. */
+	bool ValidatePassword(FStringView InAttempt);
+
+	/** Retrieves the plain text; only supported on platforms where the credential can be encrypted at rest. */
+	const UTF8CHAR* GetAuthPassword() const;
+
+	/** True if a password hash is set, even if `GetAuthPassword` would return unset. */
+	bool IsAuthPasswordSet() const;
+
+
+	/** Returns a JWT, signed with our private key, containing the provided payload. */
+	FString IssueJWT(const TMap<FStringView, FStringView>& InPayloadFields);
+
+	/** Returns the decoded and validated JWT; result is unset if validation fails. */
+	TOptional<FJsonWebToken> GetValidatedJWT(FStringView InEncodedJwt) const;
+
+	/** Returns just whether the JWT passes validation. */
+	bool IsValidJWT(FStringView InEncodedJwt) const;
+
+
+private:
+	bool CheckOrCreateSelfSignedKeyPair(FSettings& InOutSettings);
+	FUtf8String GenerateHashSalt();
+	bool UpdateLastPasswordChangeTimeUtc(FDateTime InChangeTimeUtc = FDateTime::UtcNow());
+
+private:
+	UE::SwitchboardListener::ICredentialManager* CredentialManager;
+	bool bSecureHeapNeedsShutdown = false;
+
+	FUtf8String CertificateFilePath;
+	FUtf8String PrivateKeyFilePath;
+	UTF8CHAR* PrivateKeyPassword = nullptr;
+	EVP_PKEY* PrivateKey = nullptr;
+	X509* Certificate = nullptr;
+
+	UTF8CHAR* PlainAuthPassword = nullptr;
+	FUtf8String HashSalt;
+	TArray<uint8> HashedAuthPassword;
+	int64 LastPasswordChangeTimeUtc = -1;
+};
+
+
+namespace UE::SwitchboardListener::Private
+{
+#if !PLATFORM_WINDOWS
+	static constexpr mode_t CredentialFilePermissionFlags = S_IRUSR | S_IWUSR;
+	static constexpr mode_t CredentialDirPermissionFlags = S_IRUSR | S_IWUSR | S_IXUSR;
+	static constexpr mode_t CredentialPermissionMask = S_IRWXU | S_IRWXG | S_IRWXO;
+#endif
+
+	bool CheckCredentialPermissions(FStringView InFilePath);
+	TPair<EVP_PKEY*, X509*> CreateSelfSignedPair();
+};
diff --git a/Engine/Source/Programs/SwitchboardListener/SblCore/Internal/SwitchboardCredential.h b/Engine/Source/Programs/SwitchboardListener/SblCore/Internal/SwitchboardCredential.h
new file mode 100644
--- /dev/null
+++ b/Engine/Source/Programs/SwitchboardListener/SblCore/Internal/SwitchboardCredential.h
@@ -0,0 +1,5 @@
+// Copyright Epic Games, Inc. All Rights Reserved.
+
+#pragma once
+
+#include COMPILED_PLATFORM_HEADER(SwitchboardCredential.h)
diff --git a/Engine/Source/Programs/SwitchboardListener/SblCore/Internal/SwitchboardCredentialInterface.h b/Engine/Source/Programs/SwitchboardListener/SblCore/Internal/SwitchboardCredentialInterface.h
new file mode 100644
--- /dev/null
+++ b/Engine/Source/Programs/SwitchboardListener/SblCore/Internal/SwitchboardCredentialInterface.h
@@ -0,0 +1,42 @@
+// Copyright Epic Games, Inc. All Rights Reserved.
+
+#pragma once
+
+#include "Containers/StringFwd.h"
+#include "HAL/PreprocessorHelpers.h"
+#include "Templates/SharedPointer.h"
+
+
+namespace UE::SwitchboardListener
+{
+	class ICredential : public TSharedFromThis<ICredential>
+	{
+	public:
+		virtual ~ICredential() = default;
+
+		/** Returns the user field value for this credential. */
+		virtual FStringView GetUser() = 0;
+
+		/** Returns the blob field value (password, token, etc) for this credential. */
+		virtual FStringView GetBlob() = 0;
+	};
+
+	class ICredentialManager : public TSharedFromThis<ICredentialManager>
+	{
+	public:
+		/** Returns the default credential manager for this platform. */
+		static ICredentialManager& GetPlatformCredentialManager();
+
+	public:
+		virtual ~ICredentialManager() = default;
+
+		/** Returns whether credentials are encrypted at rest on the current platform. */
+		virtual bool IsEncryptedAtRest() const { return false; }
+
+		/** Returns the stored credential if it exists. */
+		virtual TSharedPtr<ICredential> LoadCredential(FStringView InCredentialName) = 0;
+
+		/** Stores the provided credential. */
+		virtual bool SaveCredential(FStringView InCredentialName, FStringView InUser, FStringView InBlob) = 0;
+	};
+};
diff --git a/Engine/Source/Programs/SwitchboardListener/SblCore/Internal/SwitchboardListener.h b/Engine/Source/Programs/SwitchboardListener/SblCore/Internal/SwitchboardListener.h
--- a/Engine/Source/Programs/SwitchboardListener/SblCore/Internal/SwitchboardListener.h
+++ b/Engine/Source/Programs/SwitchboardListener/SblCore/Internal/SwitchboardListener.h
@@ -3,6 +3,7 @@
 #pragma once
 
 #include "SwitchboardListenerApp.h"
+#include "SwitchboardAuth.h"
 #include "Async/Mutex.h"
 #include "Async/RecursiveMutex.h"
 #include "Containers/Queue.h"
@@ -49,22 +50,6 @@ struct FSwitchboardCommandLineOptions
 	TOptional<FIPv4Address> Address;
 	TOptional<uint16> Port;
 
-	enum class ESecureMode : uint8
-	{
-		Unspecified = 0,
-		CertificateFromFile,
-#if PLATFORM_WINDOWS
-		CertificateByHash,
-#endif
-	};
-
-	ESecureMode SecureMode;
-	TOptional<FString> CertificateFile;
-	TOptional<FString> PrivateKeyFile;
-#if PLATFORM_WINDOWS
-	TOptional<FString> CertificateHash;
-#endif
-
 	TOptional<uint32> RedeployFromPid;
 
 	static FSwitchboardCommandLineOptions FromString(const TCHAR* CommandLine);
@@ -128,13 +113,12 @@ class FSwitchboardListener
 	using FByteArrayRef = TSharedRef<TArray<uint8>>;
 
 	static const FIPv4Endpoint InvalidEndpoint;
-	static constexpr const TCHAR* PasswordCredentialName = TEXT("PresharedAuthToken");
 
 public:
 	explicit FSwitchboardListener(const FSwitchboardCommandLineOptions& InOptions);
 	~FSwitchboardListener();
 
-	bool Init(bool bStartListeningImmediately = true);
+	bool Init();
 	void Shutdown();
 
 	void Tick();
@@ -143,7 +127,9 @@ public:
 	bool StopListening();
 	bool IsListening() const { return QuicApi != nullptr; }
 
-	const FString& GetAuthPassword() const { return ExpectedAuthPassword; }
+	/** May be null if the plain text password cannot be stored encrypted at rest on this platform. */
+	const UTF8CHAR* GetAuthPassword() const;
+	bool IsAuthPasswordSet() const;
 	bool SetAuthPassword(const FString& NewPassword);
 
 	TSet<FIPv4Endpoint> GetConnectedClientEndpoints() const;
@@ -249,8 +235,7 @@ private:
 	FSwitchboardCommandLineOptions Options;
 	TUniquePtr<FIPv4Endpoint> ListenerEndpoint;
 
-	TOptional<FString> PrivateKeyPassword;
-	FString ExpectedAuthPassword;
+	FSwitchboardAuthHelper AuthHelper;
 
 	/** MsQuic top level function table for all other API calls. */
 	const QUIC_API_TABLE* QuicApi = nullptr;
diff --git a/Engine/Source/Programs/SwitchboardListener/SblCore/Internal/SwitchboardListenerVersion.h b/Engine/Source/Programs/SwitchboardListener/SblCore/Internal/SwitchboardListenerVersion.h
--- a/Engine/Source/Programs/SwitchboardListener/SblCore/Internal/SwitchboardListenerVersion.h
+++ b/Engine/Source/Programs/SwitchboardListener/SblCore/Internal/SwitchboardListenerVersion.h
@@ -4,4 +4,4 @@
 
 #define SBLISTENER_VERSION_MAJOR 3
 #define SBLISTENER_VERSION_MINOR 0
-#define SBLISTENER_VERSION_PATCH 5
+#define SBLISTENER_VERSION_PATCH 6
diff --git a/Engine/Source/Programs/SwitchboardListener/SblCore/Internal/Windows/WindowsSwitchboardCredential.h b/Engine/Source/Programs/SwitchboardListener/SblCore/Internal/Windows/WindowsSwitchboardCredential.h
new file mode 100644
--- /dev/null
+++ b/Engine/Source/Programs/SwitchboardListener/SblCore/Internal/Windows/WindowsSwitchboardCredential.h
@@ -0,0 +1,28 @@
+// Copyright Epic Games, Inc. All Rights Reserved.
+
+#pragma once
+
+#include "SwitchboardCredentialInterface.h"
+
+
+namespace UE::SwitchboardListener::Private
+{
+	class FWindowsCredentialManager : public ICredentialManager
+	{
+	public:
+		virtual bool IsEncryptedAtRest() const override { return true; }
+		virtual TSharedPtr<ICredential> LoadCredential(FStringView InCredentialName) override;
+		virtual bool SaveCredential(FStringView InCredentialName, FStringView InUser, FStringView InBlob) override;
+	};
+};
+
+
+namespace UE::SwitchboardListener
+{
+	inline ICredentialManager& ICredentialManager::GetPlatformCredentialManager()
+	{
+		using namespace UE::SwitchboardListener::Private;
+		static FWindowsCredentialManager CredentialManager;
+		return CredentialManager;
+	}
+};
diff --git a/Engine/Source/Programs/SwitchboardListener/SblCore/Private/PosixOS/PosixOSSwitchboardCredential.cpp b/Engine/Source/Programs/SwitchboardListener/SblCore/Private/PosixOS/PosixOSSwitchboardCredential.cpp
new file mode 100644
--- /dev/null
+++ b/Engine/Source/Programs/SwitchboardListener/SblCore/Private/PosixOS/PosixOSSwitchboardCredential.cpp
@@ -0,0 +1,231 @@
+// Copyright Epic Games, Inc. All Rights Reserved.
+
+#include "PosixOS/PosixOSSwitchboardCredential.h"
+#include "SwitchboardAuth.h"
+#include "SwitchboardListenerApp.h" // for LogSwitchboard
+
+#include "HAL/PlatformFile.h"
+#include "HAL/PlatformFileManager.h"
+#include "Logging/StructuredLog.h"
+#include "Misc/FileHelper.h"
+#include "Misc/Paths.h"
+#include "Misc/ScopeExit.h"
+#include "Policies/PrettyJsonPrintPolicy.h"
+#include "Serialization/JsonReader.h"
+#include "Serialization/JsonSerializer.h"
+#include "Serialization/JsonWriter.h"
+
+#include <stdlib.h>
+#include <sys/stat.h>
+
+
+namespace UE::SwitchboardListener::Private
+{
+	struct FJsonCredential
+	{
+		static constexpr FStringView UserFieldName = TEXTVIEW("user");
+		static constexpr FStringView BlobFieldName = TEXTVIEW("blob");
+
+		FString UserField;
+		FString BlobField;
+
+
+		FString ToJsonString() const
+		{
+			FString Result;
+
+			TSharedRef<TJsonWriter<TCHAR, TPrettyJsonPrintPolicy<TCHAR>>> JsonWriter =
+				TJsonWriterFactory<TCHAR, TPrettyJsonPrintPolicy<TCHAR>>::Create(&Result);
+
+			JsonWriter->WriteObjectStart();
+			JsonWriter->WriteValue(UserFieldName, UserField);
+			JsonWriter->WriteValue(BlobFieldName, BlobField);
+			JsonWriter->WriteObjectEnd();
+			ensure(JsonWriter->Close());
+
+			return Result;
+		}
+
+
+		static TOptional<FJsonCredential> FromJsonString(FStringView InJsonString)
+		{
+			TSharedRef<TJsonReader<TCHAR>> Reader =
+				TJsonReaderFactory<>::CreateFromView(InJsonString);
+
+			TSharedPtr<FJsonObject> JsonData;
+			if (!ensure(FJsonSerializer::Deserialize(Reader, JsonData)))
+			{
+				UE_LOGFMT(LogSwitchboard, Error, "{Func}: Failed to deserialize: {JsonString}",
+					__FUNCTION__, InJsonString);
+				return {};
+			}
+
+			FJsonCredential Result;
+			if (!ensure(JsonData->TryGetStringField(UserFieldName, Result.UserField)))
+			{
+				UE_LOGFMT(LogSwitchboard, Error, "{Func}: Failed to get field '{FieldName}'",
+					__FUNCTION__, UserFieldName);
+				return {};
+			}
+
+			if (!ensure(JsonData->TryGetStringField(BlobFieldName, Result.BlobField)))
+			{
+				UE_LOGFMT(LogSwitchboard, Error, "{Func}: Failed to get field '{FieldName}'",
+					__FUNCTION__, BlobFieldName);
+				return {};
+			}
+
+			return Result;
+		}
+	};
+
+
+	class FPosixCredential : public ICredential
+	{
+	public:
+		FPosixCredential(const FJsonCredential& InJsonCredential)
+			: JsonCredential(InJsonCredential)
+		{
+		}
+
+		virtual FStringView GetUser() override
+		{
+			return JsonCredential.UserField;
+		}
+
+
+		virtual FStringView GetBlob() override
+		{
+			return JsonCredential.BlobField;
+		}
+
+	private:
+		FJsonCredential JsonCredential;
+	};
+
+
+	FString GetCredentialFilePath(FStringView InCredentialName)
+	{
+		const FString SblDataDir = GetSwitchboardDataDirectory();
+
+		const FString FileName = FString::Printf(TEXT("credential_%*s.json"),
+			InCredentialName.Len(), InCredentialName.GetData());
+
+		return SblDataDir / FileName;
+	}
+
+	TSharedPtr<ICredential> FPosixCredentialManager::LoadCredential(FStringView InCredentialName)
+	{
+		const FString FilePath = GetCredentialFilePath(InCredentialName);
+		if (!FPaths::FileExists(FilePath))
+		{
+			UE_LOGFMT(LogSwitchboard, Verbose, "Credential at {Path} does not exist", FilePath);
+			return {};
+		}
+
+		if (!ensure(UE::SwitchboardListener::Private::CheckCredentialPermissions(FilePath)))
+		{
+			UE_LOGFMT(LogSwitchboard, Error, "Credential file at {Path} has incorrect permissions!", FilePath);
+		}
+
+		FString CredentialStr;
+		if (!ensure(FFileHelper::LoadFileToString(CredentialStr, *FilePath)))
+		{
+			UE_LOGFMT(LogSwitchboard, Error, "Unable to read credential file at {Path}", FilePath);
+			return {};
+		}
+
+		TOptional<FJsonCredential> LoadedCredential = FJsonCredential::FromJsonString(CredentialStr);
+		if (!ensure(LoadedCredential))
+		{
+			UE_LOGFMT(LogSwitchboard, Error, "Error parsing file at {Path}", FilePath);
+			return {};
+		}
+
+		return MakeShared<FPosixCredential>(*LoadedCredential);
+	}
+
+
+	bool FPosixCredentialManager::SaveCredential(FStringView InCredentialName, FStringView InUser, FStringView InBlob)
+	{
+		// Construct the temporary file path template for mkstemp.
+		const FString DestFilePath = GetCredentialFilePath(InCredentialName);
+		const FString DestDir = FPaths::GetPath(DestFilePath);
+		const FUtf8String TempFilePathTemplate(DestDir / FString("sbltempXXXXXX"));
+
+		// Copy into a mutable array for mkstemp(), which modifies its argument
+		// in-place, replacing 'XXXXXX' with randomly-generated unique characters.
+		TArray<UTF8CHAR> TempFilePathU8(*TempFilePathTemplate, TempFilePathTemplate.Len() + 1);
+
+		// Create the directory tree with appropriate permissions.
+		const mode_t OriginalUmask = umask(~CredentialDirPermissionFlags);
+		ON_SCOPE_EXIT{ umask(OriginalUmask); };
+
+		IPlatformFile& PlatformFile = FPlatformFileManager::Get().GetPlatformFile();
+		PlatformFile.CreateDirectoryTree(*DestDir);
+
+		// Create the temporary file with appropriately restrictive permissions.
+		// POSIX.1-2008 says this umask should be redundant, but just to be safe.
+		umask(~CredentialFilePermissionFlags);
+		int TempFd = mkstemp(reinterpret_cast<char*>(TempFilePathU8.GetData()));
+		if (!ensure(TempFd != -1))
+		{
+			const int ErrorCode = errno;
+			UE_LOGFMT(LogSwitchboard, Error, "{Func}: mkstemp failed (errno: %d; path: '%S')",
+				__FUNCTION__, ErrorCode, TempFilePathU8);
+			return false;
+		}
+
+		const FJsonCredential JsonCredential{ FString(InUser), FString(InBlob) };
+		const FUtf8String JsonString(JsonCredential.ToJsonString());
+
+		int ErrorCode = 0;
+		int32 TotalBytesWritten = 0;
+		while (TotalBytesWritten < JsonString.Len())
+		{
+			const SSIZE_T Written = write(TempFd,
+				(*JsonString) + TotalBytesWritten,
+				JsonString.Len() - TotalBytesWritten);
+
+			if (Written <= 0)
+			{
+				ErrorCode = errno;
+				break;
+			}
+
+			TotalBytesWritten += Written;
+		}
+
+		const int CloseResult = close(TempFd);
+		if (!ensure(CloseResult == 0))
+		{
+			const int CloseErrorCode = errno;
+			UE_LOGFMT(LogSwitchboard, Warning, "{Func}: close failed (errno: %d; path: '%S')",
+				__FUNCTION__, CloseErrorCode, TempFilePathU8);
+		}
+
+		// If write() failed, delete the truncated temporary file.
+		if (TotalBytesWritten < JsonString.Len())
+		{
+			PlatformFile.DeleteFile(UTF8_TO_TCHAR(TempFilePathU8.GetData()));
+			UE_LOGFMT(LogSwitchboard, Error,
+				"{Func}: write failed (errno: %d; written %d/%d; path: '%S')",
+				__FUNCTION__, ErrorCode, TotalBytesWritten, JsonString.Len(), TempFilePathU8);
+			return false;
+		}
+
+		// Otherwise, rename the temporary file to its intended destination.
+		const int RenameResult = rename((const char*)TempFilePathU8.GetData(),
+			TCHAR_TO_UTF8(*DestFilePath));
+		if (!ensure(RenameResult == 0))
+		{
+			ErrorCode = errno;
+			PlatformFile.DeleteFile(UTF8_TO_TCHAR(TempFilePathU8.GetData()));
+			UE_LOGFMT(LogSwitchboard, Warning, "{Func}: rename failed (errno: %d; path: '%S' -> '%s')",
+				__FUNCTION__, ErrorCode, TempFilePathU8, DestFilePath);
+			return false;
+		}
+
+		return true;
+	}
+}; // namespace UE::SwitchboardListener::Private
diff --git a/Engine/Source/Programs/SwitchboardListener/SblCore/Private/SwitchboardAuth.cpp b/Engine/Source/Programs/SwitchboardListener/SblCore/Private/SwitchboardAuth.cpp
--- a/Engine/Source/Programs/SwitchboardListener/SblCore/Private/SwitchboardAuth.cpp
+++ b/Engine/Source/Programs/SwitchboardListener/SblCore/Private/SwitchboardAuth.cpp
@@ -1,15 +1,27 @@
 // Copyright Epic Games, Inc. All Rights Reserved.
 
 #include "SwitchboardAuth.h"
+#include "SwitchboardCredential.h"
 #include "SwitchboardListenerApp.h"
 
+#include "Containers/StaticArray.h"
 #include "Containers/StringConv.h"
 #include "HAL/FileManagerGeneric.h"
 #include "HAL/PlatformProcess.h"
+#include "IJwt.h"
+#include "IPlatformCrypto.h"
+#include "JsonWebToken.h"
+#include "JwtAlgorithms.h"
 #include "Logging/StructuredLog.h"
+#include "Misc/Base64.h"
 #include "Misc/FileHelper.h"
 #include "Misc/Paths.h"
 #include "Misc/ScopeExit.h"
+#include "Misc/StringBuilder.h"
+#include "Policies/CondensedJsonPrintPolicy.h"
+#include "Serialization/JsonWriter.h"
+#include "String/BytesToHex.h"
+#include "String/HexToBytes.h"
 
 #include <iostream>
 #include <string>
@@ -18,26 +30,210 @@
 #include "Windows/AllowWindowsPlatformTypes.h"
 #include <Windows.h>
 #include <ShlObj.h>
-#include <wincred.h>
 #include "Windows/HideWindowsPlatformTypes.h"
 #else
+#include <fcntl.h>
 #include <sys/file.h>
+#include <sys/stat.h>
+#include <sys/types.h>
 #include <termios.h>
+#include <unistd.h>
 #endif
 
-UE_PUSH_MACRO("UI")
-// Workaround for "ossl_typ.h(144): error C2365: 'UI': redefinition; previous definition was 'namespace'"
-//                "ObjectMacros.h(872): note: see declaration of 'UI'"
-#define UI UI_ST
+#include <openssl/crypto.h>
+#include <openssl/kdf.h>
 #include <openssl/pem.h>
 #include <openssl/rand.h>
 #include <openssl/x509.h>
-UE_POP_MACRO("UI")
+
+
+namespace UE::SwitchboardListener::Private
+{
+	constexpr FStringView PlainPasswordCredentialName = TEXTVIEW("PresharedAuthToken");
+	constexpr FStringView PasswordSaltCredentialName = TEXTVIEW("PasswordSalt");
+	constexpr FStringView PasswordHashCredentialName = TEXTVIEW("PasswordHash");
+	constexpr FStringView LastPasswordChangeTimeCredentialName = TEXTVIEW("PasswordLastChanged");
+
+	constexpr FStringView JwtIssuer_Value = TEXTVIEW("SwitchboardListener");
+
+	constexpr SIZE_T ScryptNumDerivedBytes = 64;
+
+	// The file should only be accessible for read/write by the current user.
+	// If the permissions are not as expected, we need to alert the user.
+	bool CheckCredentialPermissions(FStringView InFilePath)
+	{
+#if PLATFORM_WINDOWS
+		return true;
+#else
+		FString NormalizedPath(InFilePath);
+		FPaths::NormalizeFilename(NormalizedPath);
+		NormalizedPath = FPaths::ConvertRelativePathToFull(NormalizedPath);
+
+		// lstat is used here to explicitly not follow symbolic links.
+		struct stat Stat;
+		const int StatResult = lstat(TCHAR_TO_UTF8(*NormalizedPath), &Stat);
+		if (StatResult == 0)
+		{
+			// Note: getuid is used here instead of geteuid to explicitly get the real user ID,
+			// and not allow the file to be owned by another effective user ID running the program.
+			const uid_t CurrentUserId = getuid();
+			if (CurrentUserId != Stat.st_uid)
+			{
+				UE_LOGFMT(LogSwitchboard, Error,
+					"The stored credential file is not owned by the current user. (Current: {CurrentUserId}; Owner: {st_uid})",
+					CurrentUserId, Stat.st_uid);
+
+				return false;
+			}
+
+			// Filter out rest of the bit flags and only check against the
+			// User, Group and Other groups for their Read/Write/Execute access rights.
+			const mode_t ShareFlags = Stat.st_mode & CredentialPermissionMask;
+			if (ShareFlags == CredentialFilePermissionFlags)
+			{
+				return true;
+			}
+			else
+			{
+				UE_LOGFMT(LogSwitchboard, Error,
+					"The stored credential file has insecure permissions. (Current: {ShareFlags}; Expected: {CredentialFlags})",
+					ShareFlags, CredentialFilePermissionFlags);
+			}
+		}
+		else
+		{
+			const int ErrorCode = errno;
+			if (ErrorCode == ENOENT)
+			{
+				UE_LOGFMT(LogSwitchboard, Verbose, "Credential file does not exist. (Path: '{Path}')", *NormalizedPath);
+			}
+			else
+			{
+				UE_LOGFMT(LogSwitchboard, Error, "lstat failed. (errno: {ErrorCode}; Path: '{Path}')", ErrorCode, *NormalizedPath);
+			}
+		}
+
+		return false;
+#endif
+	}
+
+	/**
+	 * On POSIX, use open + fdopen to ensure the file is created
+	 * atomically with the correct, restrictive permissions.
+	 */
+	struct FSensitiveFile
+	{
+		FILE* Fp = nullptr;
+#if !PLATFORM_WINDOWS
+		int Fd = -1;
+#endif
+
+		~FSensitiveFile()
+		{
+#if PLATFORM_WINDOWS
+			const bool bNeedsClose = Fp != nullptr;
+#else
+			const bool bNeedsClose = Fp != nullptr || Fd != -1;
+#endif
+
+			if (bNeedsClose)
+			{
+				ensure(Close());
+			}
+		}
+
+		FILE* OpenForWrite(FStringView InFilePath)
+		{
+#if PLATFORM_WINDOWS
+			const errno_t Result = _wfopen_s(&Fp,
+				TCHAR_TO_WCHAR(InFilePath.GetData()),
+				L"w");
+
+			if (Result != 0)
+			{
+				UE_LOGFMT(LogSwitchboard, Error, "_wfopen_s failed ({Result})", Result);
+				Fp = nullptr;
+			}
+
+			return Fp;
+#else
+			Fd = open(TCHAR_TO_UTF8(InFilePath.GetData()),
+				O_CREAT | O_EXCL | O_WRONLY,
+				CredentialFilePermissionFlags);
+
+			if (Fd == -1)
+			{
+				const int OpenErrno = errno;
+				UE_LOGFMT(LogSwitchboard, Error, "open failed (errno: {OpenErrno})", OpenErrno);
+				return nullptr;
+			}
+
+			Fp = fdopen(Fd, "w");
+			if (Fp == nullptr)
+			{
+				const int FdopenErrno = errno;
+				UE_LOGFMT(LogSwitchboard, Error, "fdopen failed (errno: {FdopenErrno})", FdopenErrno);
+				close(Fd);
+				return nullptr;
+			}
+
+			return Fp;
+#endif
+		}
+
+		bool Close()
+		{
+			bool bSuccess = true;
+
+			if (Fp != nullptr)
+			{
+				if (!ensure(fclose(Fp)) == 0)
+				{
+					const int FcloseErrno = errno;
+					UE_LOGFMT(LogSwitchboard, Error, "fclose failed (errno: {FcloseErrno})", FcloseErrno);
+					bSuccess = false;
+				}
+
+				Fp = nullptr;
+			}
+
+#if !PLATFORM_WINDOWS
+			if (Fd != -1)
+			{
+				if (!ensure(close(Fd) == 0))
+				{
+					const int CloseErrno = errno;
+					UE_LOGFMT(LogSwitchboard, Error, "close failed (errno: {CloseErrno})", CloseErrno);
+					bSuccess = false;
+				}
+				Fd = -1;
+			}
+#endif
+
+			return bSuccess;
+		}
+	};
+}
 
 
 namespace UE::SwitchboardListener
 {
-	FString ReadCredentialFromStdin(const FString& CredentialDesc, bool bAllowEmpty /* = false */)
+	FString GetSwitchboardDataDirectory()
+	{
+#if PLATFORM_WINDOWS
+		wchar_t* LocalAppDataWstr = nullptr;
+		::SHGetKnownFolderPath(FOLDERID_LocalAppData, 0, nullptr, &LocalAppDataWstr);
+		ON_SCOPE_EXIT{ ::CoTaskMemFree(LocalAppDataWstr); };
+
+		const FString LocalAppDataDir(LocalAppDataWstr);
+		return LocalAppDataDir / TEXT("SwitchboardListener");
+#else
+		const FString UserHomeDir = FPlatformProcess::UserHomeDir();
+		return UserHomeDir / TEXT(".switchboard") / TEXT("listener");
+#endif
+	}
+
+	FString ReadPasswordFromStdin(bool bAllowEmpty /* = false */)
 	{
 		FString Credential;
 
@@ -70,16 +266,16 @@ namespace UE::SwitchboardListener
 			// FIXME?: Doesn't seem like we have a suitable abstraction for this?
 			std::string CredFirst, CredSecond;
 
-			printf("\nEnter %s: ", TCHAR_TO_UTF8(*CredentialDesc));
+			printf("\nEnter password: ");
 			std::getline(std::cin, CredFirst);
 
 			if (CredFirst.empty() && !bAllowEmpty)
 			{
-				printf("\nEmpty %s is not allowed!", TCHAR_TO_UTF8(*CredentialDesc));
+				printf("\nEmpty password is not allowed!");
 				continue;
 			}
 
-			printf("\nConfirm %s: ", TCHAR_TO_UTF8(*CredentialDesc));
+			printf("\nConfirm password: ");
 			std::getline(std::cin, CredSecond);
 
 			if (CredFirst == CredSecond)
@@ -96,91 +292,12 @@ namespace UE::SwitchboardListener
 		return Credential;
 	}
 
-	bool SupportsPersistentCredentials()
-	{
-#if PLATFORM_WINDOWS
-		return true;
-#else
-		return false;
-#endif
-	}
-
-	bool SaveCredential(const FCredential& Credential)
-	{
-#if PLATFORM_WINDOWS
-		const FString NamespacedTargetName =
-			FString::Printf(TEXT("SwitchboardListener_%s"), *Credential.CredentialName);
-
-		CREDENTIALW Cred = { 0 };
-		Cred.Type = CRED_TYPE_GENERIC;
-		Cred.TargetName = const_cast<FString::ElementType*>(*NamespacedTargetName);
-		Cred.CredentialBlobSize = sizeof(TCHAR) * Credential.CredentialBlob.Len();
-		Cred.CredentialBlob = reinterpret_cast<BYTE*>(
-			const_cast<FString::ElementType*>(*Credential.CredentialBlob));
-		Cred.Persist = CRED_PERSIST_LOCAL_MACHINE;
-		Cred.UserName = const_cast<FString::ElementType*>(*Credential.UserName);
-		if (::CredWriteW(&Cred, 0))
-		{
-			return true;
-		}
-		else
-		{
-			UE_LOGFMT(LogSwitchboard, Warning, "CredWriteW failed; GetLastError = {LastError}", static_cast<int64>(::GetLastError()));
-		}
-#else
-		// Not implemented
-#endif
-
-		return false;
-	}
-
-	TOptional<FCredential> LoadCredential(FStringView CredentialName)
-	{
-#if PLATFORM_WINDOWS
-		const FString NamespacedTargetName =
-			FString::Printf(TEXT("SwitchboardListener_%*s"),
-				CredentialName.Len(), CredentialName.GetData());
-
-		CREDENTIALW* OutCredPtr = nullptr;
-		if (::CredReadW(*NamespacedTargetName, CRED_TYPE_GENERIC, 0, &OutCredPtr))
-		{
-			FCredential Result;
-			Result.CredentialName = CredentialName;
-			Result.UserName = FString(OutCredPtr->UserName);
-			Result.CredentialBlob = FString(
-				OutCredPtr->CredentialBlobSize / sizeof(TCHAR),
-				reinterpret_cast<TCHAR*>(OutCredPtr->CredentialBlob));
-
-			::CredFree(OutCredPtr);
-
-			return Result;
-		}
-		else
-		{
-			const int64 LastError = ::GetLastError();
-			if (LastError == ERROR_NOT_FOUND)
-			{
-				UE_LOGFMT(LogSwitchboard, Verbose, "CredReadW returned ERROR_NOT_FOUND");
-			}
-			else
-			{
-				UE_LOGFMT(LogSwitchboard, Warning, "CredReadW failed; GetLastError = {LastError}", static_cast<int64>(::GetLastError()));
-			}
-		}
-#else
-		// Not implemented
-#endif
-
-		return {};
-	}
-
-
-	void FillSecureRandom(TArray<uint8>& InOutArray)
+	void FillSecureRandom(TArrayView<uint8> InOutArray)
 	{
 		const int32 NumBytes = InOutArray.Num();
 		if (ensure(NumBytes > 0))
 		{
-			check(1 == RAND_bytes(
+			check(1 == RAND_priv_bytes(
 				reinterpret_cast<unsigned char*>(InOutArray.GetData()),
 				NumBytes));
 		}
@@ -188,7 +305,7 @@ namespace UE::SwitchboardListener
 }
 
 
-namespace UE::SwitchboardListener::Certificates::Private
+namespace UE::SwitchboardListener::Private
 {
 	/**
 	 * Generate a 2048-bit RSA key.
@@ -221,10 +338,12 @@ namespace UE::SwitchboardListener::Certificates::Private
 	}
 
 
-	FString BIOToString(BIO* Bio)
+	FString MemBIOToString(BIO* InBio)
 	{
+		check(InBio);
+
 		char* Text;
-		long Length = BIO_get_mem_data(Bio, &Text);
+		long Length = BIO_get_mem_data(InBio, &Text);
 
 		auto Convert = StringCast<TCHAR>(Text, Length);
 		FString Result(Convert.Length(), Convert.Get());
@@ -233,48 +352,110 @@ namespace UE::SwitchboardListener::Certificates::Private
 	}
 
 
+	BIO* StringToMemBIO(const FString& InString)
+	{
+		return BIO_new_mem_buf(TCHAR_TO_UTF8(*InString), -1);
+	}
+
+
 	/**
-	 * Convert a certificate to a string.
+	 * Convert a certificate to a PEM format string.
 	 *
-	 * @param x509 The certificate struct
+	 * @param InX509 The certificate struct.
 	 */
-	FString X509ToString(X509* x509)
+	FString CertificateToString(X509* InX509)
 	{
+		check(InX509);
+
 		BIO* Bio = BIO_new(BIO_s_mem());
 		ON_SCOPE_EXIT{ BIO_free(Bio); };
 
-		check(1 == PEM_write_bio_X509(Bio, x509));
+		check(1 == PEM_write_bio_X509(Bio, InX509));
 
-		return BIOToString(Bio);
+		return MemBIOToString(Bio);
 	}
 
 
 	/**
-	 * Converts a private key to a string.
+	 * Convert a certificate's public key to a PEM format string.
 	 *
-	 * @param PrivateKey The private key struct
+	 * @param InX509 The certificate struct.
 	 */
-	FString PrivateKeyToString(EVP_PKEY* PrivateKey, const FString& PrivateKeyPassword)
+	FString CertPubKeyToString(X509* InX509)
 	{
+		check(InX509);
+
 		BIO* Bio = BIO_new(BIO_s_mem());
 		ON_SCOPE_EXIT{ BIO_free(Bio); };
 
-		if (PrivateKeyPassword.IsEmpty())
+		EVP_PKEY* PubKey = X509_get0_pubkey(InX509);
+		check(PubKey);
+
+		check(1 == PEM_write_bio_PUBKEY(Bio, PubKey));
+
+		return MemBIOToString(Bio);
+	}
+
+
+	bool SavePrivateKeyToFile(EVP_PKEY* InPrivateKey, const FString& InFilePath, const UTF8CHAR* InPrivateKeyPassword)
+	{
+		FSensitiveFile DestFile;
+		FILE* DestFp = DestFile.OpenForWrite(InFilePath);
+		if (!ensure(DestFp))
+		{
+			return false;
+		}
+
+		ON_SCOPE_EXIT{ ensure(DestFile.Close()); };
+
+		if (!InPrivateKeyPassword || !InPrivateKeyPassword[0])
 		{
-			check(1 == PEM_write_bio_PrivateKey_traditional(Bio, PrivateKey,
+			return 1 == PEM_write_PrivateKey(DestFp, InPrivateKey,
 				nullptr,
 				nullptr, 0, nullptr,
-				nullptr));
+				nullptr);
 		}
 		else
 		{
-			check(1 == PEM_write_bio_PrivateKey_traditional(Bio, PrivateKey,
+			return 1 == PEM_write_PrivateKey(DestFp, InPrivateKey,
 				EVP_des_ede3_cbc(),
 				nullptr, 0, nullptr,
-				TCHAR_TO_UTF8(*PrivateKeyPassword)));
+				(void*)InPrivateKeyPassword);
 		}
+	}
+
+
+	X509* LoadX509FromFile(const FString& InFilePath)
+	{
+#if PLATFORM_WINDOWS
+		FILE* File = nullptr;
+		ensure(0 == fopen_s(&File, TCHAR_TO_UTF8(*InFilePath), "r"));
+#else
+		FILE* File = fopen(TCHAR_TO_UTF8(*InFilePath), "r");
+#endif
+		ON_SCOPE_EXIT{ fclose(File); };
 
-		return BIOToString(Bio);
+		return PEM_read_X509(File, nullptr, nullptr, nullptr);
+	}
+
+
+	EVP_PKEY* LoadPrivateKeyFromFile(const FString& InFilePath, const UTF8CHAR* InPrivateKeyPassword)
+	{
+		if (!CheckCredentialPermissions(InFilePath))
+		{
+			UE_LOGFMT(LogSwitchboard, Error, "{Func}: Insecure permissions on private key file! ({Path})",
+				__FUNCTION__, InFilePath);
+		}
+
+#if PLATFORM_WINDOWS
+		FILE* File = nullptr;
+		ensure(0 == fopen_s(&File, TCHAR_TO_UTF8(*InFilePath), "r"));
+#else
+		FILE* File = fopen(TCHAR_TO_UTF8(*InFilePath), "r");
+#endif
+		ON_SCOPE_EXIT{ fclose(File); };
+
+		return PEM_read_PrivateKey(File, nullptr, nullptr, (void*)InPrivateKeyPassword);
 	}
 
 
@@ -283,7 +464,7 @@ namespace UE::SwitchboardListener::Certificates::Private
 	 *
 	 * @param PrivateKey The private key struct
 	 */
-	X509* GenerateCertificate(EVP_PKEY* PrivateKey)
+	X509* GenerateCertificate(EVP_PKEY* InPrivateKey)
 	{
 		X509* x509 = X509_new();
 
@@ -297,7 +478,7 @@ namespace UE::SwitchboardListener::Certificates::Private
 		X509_gmtime_adj(X509_get_notAfter(x509), (60L*60*24*365*10));
 
 		// Set the public key for our certificate.
-		X509_set_pubkey(x509, PrivateKey);
+		X509_set_pubkey(x509, InPrivateKey);
 
 		// Set the common name, and copy the subject name to the issuer name.
 		X509_NAME* SubjectName = X509_get_subject_name(x509);
@@ -306,7 +487,7 @@ namespace UE::SwitchboardListener::Certificates::Private
 		X509_set_issuer_name(x509, SubjectName);
 
 		// Actually sign the certificate with our key.
-		if (!X509_sign(x509, PrivateKey, EVP_sha256()))
+		if (!X509_sign(x509, InPrivateKey, EVP_sha256()))
 		{
 			X509_free(x509);
 			return nullptr;
@@ -315,123 +496,709 @@ namespace UE::SwitchboardListener::Certificates::Private
 		return x509;
 	}
 
-#if !PLATFORM_WINDOWS
-	bool EnsurePathOnlyOwnerReadable(const FString& InPath)
+
+	/** Get the default paths for (self-signed certificate, private key). */
+	TTuple<FString, FString> GetDefaultSelfSignedPaths()
 	{
-		FString NormalizedPath = InPath;
-		FPaths::NormalizeFilename(NormalizedPath);
-		NormalizedPath = FPaths::ConvertRelativePathToFull(NormalizedPath);
+		const TCHAR* CertFilename = TEXT("transportcert.pem");
+		const TCHAR* PrivKeyFilename = TEXT("transportkey.key");
 
-		struct stat Stat;
-		if (stat(TCHAR_TO_UTF8(*NormalizedPath), &Stat) == 0)
+		const FString SblDataDir = GetSwitchboardDataDirectory();
+		const FString CertificateFile = SblDataDir / CertFilename;
+		const FString PrivateKeyFile = SblDataDir / PrivKeyFilename;
+
+		return TTuple<FString, FString>(CertificateFile, PrivateKeyFile);
+	}
+
+
+	// Returns the first 128 bits of the SHA256 fingerprint in the format "09:AB:CD:..."
+	FString GetX509TruncatedFingerprint(X509* InX509)
+	{
+		uint8 Digest[EVP_MAX_MD_SIZE];
+		uint32 DigestSize = 0;
+		check(1 == X509_digest(InX509, EVP_sha256(), Digest, &DigestSize));
+
+		FString Fingerprint;
+		const uint32 TruncatedDigestBytes = DigestSize / 2;
+		for (uint32 ByteIdx = 0; ByteIdx < TruncatedDigestBytes; ++ByteIdx)
 		{
-			// Clear group/other read/write/execute bits
-			Stat.st_mode &= ~(S_IRWXG | S_IRWXO);
-			chmod(TCHAR_TO_UTF8(*NormalizedPath), Stat.st_mode);
-			return true;
+			ByteToHex(Digest[ByteIdx], Fingerprint);
+			if (ByteIdx < (TruncatedDigestBytes - 1))
+			{
+				Fingerprint.AppendChar(TEXT(':'));
+			}
 		}
-		else
+
+		return Fingerprint;
+	}
+
+
+	/**
+	 * Generate a private key + self-signed certificate.
+	 */
+	TPair<EVP_PKEY*, X509*> CreateSelfSignedPair()
+	{
+		EVP_PKEY* PrivateKey = GeneratePrivateKey();
+		if (!ensure(PrivateKey))
+		{
+			return { nullptr, nullptr };
+		}
+
+		X509* Certificate = GenerateCertificate(PrivateKey);
+		if (!ensure(Certificate))
+		{
+			EVP_PKEY_free(PrivateKey);
+			return { nullptr, nullptr };
+		}
+
+		return { PrivateKey, Certificate };
+	}
+
+
+	bool DeriveScryptHash(FUtf8StringView InPassword, FUtf8StringView InSalt, TArrayView<uint8> OutDerivedHash)
+	{
+		if (!ensure(OutDerivedHash.Num() == ScryptNumDerivedBytes))
+		{
+			UE_LOGFMT(LogSwitchboard, Error, "{Func}: Incorrect array length (expected {ExpectedLen}, got {ArrayLen})",
+				__FUNCTION__, uint64(ScryptNumDerivedBytes), OutDerivedHash.Num());
+			return false;
+		}
+
+		// Set according to OWASP recommendation for minimal CPU cost.
+		constexpr uint64 ScryptParam_N = 1 << 17;
+		constexpr uint64 ScryptParam_r = 8;
+		constexpr uint64 ScryptParam_p = 1;
+
+		EVP_PKEY_CTX* ScryptCtx = EVP_PKEY_CTX_new_id(EVP_PKEY_SCRYPT, nullptr);
+		check(ScryptCtx);
+		ON_SCOPE_EXIT{ EVP_PKEY_CTX_free(ScryptCtx); };
+
+		if (!ensure(EVP_PKEY_derive_init(ScryptCtx) > 0))
+		{
+			UE_LOGFMT(LogSwitchboard, Error, "{Func}: EVP_PKEY_derive_init", __FUNCTION__);
+			return false;
+		}
+
+		if (!ensure(EVP_PKEY_CTX_set1_pbe_pass(ScryptCtx, InPassword.GetData(), InPassword.Len()) > 0))
+		{
+			UE_LOGFMT(LogSwitchboard, Error, "{Func}: EVP_PKEY_CTX_set1_pbe_pass", __FUNCTION__);
+			return false;
+		}
+
+		if (!ensure(EVP_PKEY_CTX_set1_scrypt_salt(ScryptCtx, InSalt.GetData(), InSalt.Len()) > 0))
+		{
+			UE_LOGFMT(LogSwitchboard, Error, "{Func}: EVP_PKEY_CTX_set1_scrypt_salt", __FUNCTION__);
+			return false;
+		}
+
+		if (!ensure(EVP_PKEY_CTX_set_scrypt_N(ScryptCtx, ScryptParam_N) > 0))
+		{
+			UE_LOGFMT(LogSwitchboard, Error, "{Func}: EVP_PKEY_CTX_set_scrypt_N", __FUNCTION__);
+			return false;
+		}
+
+		if (!ensure(EVP_PKEY_CTX_set_scrypt_r(ScryptCtx, ScryptParam_r) > 0))
+		{
+			UE_LOGFMT(LogSwitchboard, Error, "{Func}: EVP_PKEY_CTX_set_scrypt_r", __FUNCTION__);
+			return false;
+		}
+
+		if (!ensure(EVP_PKEY_CTX_set_scrypt_p(ScryptCtx, ScryptParam_p) > 0))
+		{
+			UE_LOGFMT(LogSwitchboard, Error, "{Func}: EVP_PKEY_CTX_set_scrypt_p", __FUNCTION__);
+			return false;
+		}
+
+		SIZE_T InOutDerivedLen = ScryptNumDerivedBytes;
+		if (!ensure(EVP_PKEY_derive(ScryptCtx, OutDerivedHash.GetData(), &InOutDerivedLen) > 0))
+		{
+			UE_LOGFMT(LogSwitchboard, Error, "{Func}: EVP_PKEY_derive", __FUNCTION__);
+			return false;
+		}
+
+		if (!ensure(InOutDerivedLen == ScryptNumDerivedBytes))
 		{
+			UE_LOGFMT(LogSwitchboard, Error, "{Func}: Incorrect output length from EVP_PKEY_derive (Expected {ExpectedBytes}, got {ActualBytes})",
+				__FUNCTION__, uint64(ScryptNumDerivedBytes), uint64(InOutDerivedLen));
 			return false;
 		}
+
+		return true;
 	}
-#endif
 }
 
 
-namespace UE::SwitchboardListener::Certificates
+bool FSwitchboardAuthHelper::Initialize(const FSwitchboardAuthHelper::FSettings& InSettings)
 {
-	using namespace UE::SwitchboardListener::Certificates::Private;
-
+	using namespace UE::SwitchboardListener;
+	using namespace UE::SwitchboardListener::Private;
 
-	TTuple<FString, FString> GetSelfSignedPaths(bool bCreate /* = false */)
+	if (InSettings.CredentialManager)
 	{
-		IPlatformFile& FileManager = FPlatformFileManager::Get().GetPlatformFile();
+		CredentialManager = InSettings.CredentialManager;
+	}
+	else
+	{
+		CredentialManager = &ICredentialManager::GetPlatformCredentialManager();
+	}
 
-#if PLATFORM_WINDOWS
-		wchar_t* LocalAppDataWstr = nullptr;
-		::SHGetKnownFolderPath(FOLDERID_LocalAppData, 0, nullptr, &LocalAppDataWstr);
-		ON_SCOPE_EXIT{ ::CoTaskMemFree(LocalAppDataWstr); };
+	// Mutable copy.
+	FSettings Settings = InSettings;
 
-		const FString LocalAppDataDir(LocalAppDataWstr);
-		const FString SblHomeDir = LocalAppDataDir / TEXT("SwitchboardListener");
-		if (bCreate)
+	// Try to ensure OpenSSL secure heap is initialized.
+	// Not strictly required, but without it, credentials may
+	// get swapped to disk, be included in core dumps, etc.
+	if (!CRYPTO_secure_malloc_initialized())
+	{
+		const SIZE_T SecureHeapSizeBytes = 1024 * 1024;
+		const SIZE_T SecureHeapMinAllocSizeBytes = 64;
+		if (0 == CRYPTO_secure_malloc_init(SecureHeapSizeBytes, SecureHeapMinAllocSizeBytes))
+		{
+			UE_LOGFMT(LogSwitchboard, Warning, "Couldn't initialize secure heap");
+		}
+		else
 		{
-			FileManager.CreateDirectory(*SblHomeDir);
+			bSecureHeapNeedsShutdown = true;
 		}
+	}
 
-		const FString CertificateFile = SblHomeDir / TEXT("transportcert.pem");
-		const FString PrivateKeyFile = SblHomeDir / TEXT("transportkey.key");
-#else
-		const FString UserHomeDir = FPlatformProcess::UserHomeDir();
-		const FString SblHomeDir = UserHomeDir / TEXT(".switchboardlistener");
-		const FString SblCertDir = SblHomeDir / TEXT("cert");
+	// If the user didn't specify key pair, check the default path,
+	// and possibly generate a new self-signed key pair.
+	const bool bUserSpecifiedKeyPair = Settings.Certificate && Settings.PrivateKey;
+	if (!bUserSpecifiedKeyPair)
+	{
+		if (Settings.Certificate || Settings.PrivateKey)
+		{
+			UE_LOGFMT(LogSwitchboard, Error, "{Func}: Must specify both Certificate AND PrivateKey, or neither; ignoring", __FUNCTION__);
+		}
 
-		if (bCreate)
+		if (Settings.PrivateKeyPassword)
 		{
-			FileManager.CreateDirectory(*SblHomeDir);
-			FileManager.CreateDirectory(*SblCertDir);
-			EnsurePathOnlyOwnerReadable(SblCertDir);
+			UE_LOGFMT(LogSwitchboard, Error, "{Func}: Cannot specify password for auto-generated private key; ignoring", __FUNCTION__);
+			Settings.PrivateKeyPassword.Reset();
 		}
 
-		const FString CertificateFile = SblCertDir / TEXT("transportcert.pem");
-		const FString PrivateKeyFile = SblCertDir / TEXT("transportkey.key");
-#endif
+		CheckOrCreateSelfSignedKeyPair(Settings);
+	}
 
-		return TTuple<FString, FString>(CertificateFile, PrivateKeyFile);
+	if (!Settings.PrivateKey)
+	{
+		UE_LOGFMT(LogSwitchboard, Error, "{Func}: No private key specified", __FUNCTION__);
+		return false;
+	}
+	else if (const FString* PrivateKeyPathStr = Settings.PrivateKey->TryGet<FString>())
+	{
+		PrivateKeyFilePath = **PrivateKeyPathStr;
+		UE_LOGFMT(LogSwitchboard, Verbose, "{Func}: Loading private key from file {PrivateKeyPath}",
+			__FUNCTION__, *PrivateKeyPathStr);
+		PrivateKey = LoadPrivateKeyFromFile(*PrivateKeyPathStr, PrivateKeyPassword);
+	}
+	else
+	{
+		UE_LOGFMT(LogSwitchboard, Verbose, "{Func}: Using provided EVP_PKEY object", __FUNCTION__);
+		check(Settings.PrivateKey->IsType<EVP_PKEY*>());
+		PrivateKey = Settings.PrivateKey->Get<EVP_PKEY*>();
 	}
 
+	if (!PrivateKey)
+	{
+		UE_LOGFMT(LogSwitchboard, Error, "{Func}: Failed to load private key", __FUNCTION__);
+		return false;
+	}
 
-	TOptional<FString> CreateSelfSigned(const FString& PrivateKeyPassword)
+	if (!Settings.Certificate)
 	{
-		EVP_PKEY* PrivateKey = GeneratePrivateKey();
-		if (!PrivateKey)
+		UE_LOGFMT(LogSwitchboard, Error, "{Func}: No certificate specified", __FUNCTION__);
+		return false;
+	}
+	else if (const FString* CertificatePathStr = Settings.Certificate->TryGet<FString>())
+	{
+		CertificateFilePath = **CertificatePathStr;
+		UE_LOGFMT(LogSwitchboard, Verbose, "{Func}: Loading certificate from file {CertificatePath}",
+			__FUNCTION__, *CertificatePathStr);
+		Certificate = LoadX509FromFile(*CertificatePathStr);
+	}
+	else
+	{
+		UE_LOGFMT(LogSwitchboard, Verbose, "{Func}: Using provided X509 object", __FUNCTION__);
+		check(Settings.Certificate->IsType<X509*>());
+		Certificate = Settings.Certificate->Get<X509*>();
+	}
+
+	if (!Certificate)
+	{
+		UE_LOGFMT(LogSwitchboard, Error, "{Func}: Failed to load certificate", __FUNCTION__);
+		return false;
+	}
+
+	// Finally, load the auth password (or its hash)
+	if (CredentialManager->IsEncryptedAtRest())
+	{
+		if (TSharedPtr<ICredential> SavedPassword = CredentialManager->LoadCredential(PlainPasswordCredentialName))
 		{
-			return TOptional<FString>{};
+			UE_LOGFMT(LogSwitchboard, Display, "Using stored password");
+
+			// This salt is ephemeral since we store the plain text password.
+			HashSalt = GenerateHashSalt();
+
+			const bool bExpireExistingTokens_false = false;
+			ensure(SetAuthPassword(SavedPassword->GetBlob(), bExpireExistingTokens_false));
+		}
+	}
+	else
+	{
+		if (TSharedPtr<ICredential> SavedSalt = CredentialManager->LoadCredential(PasswordSaltCredentialName))
+		{
+			UE_LOGFMT(LogSwitchboard, Display, "Using stored password salt");
+			HashSalt = FUtf8String(SavedSalt->GetBlob());
 		}
+		else
+		{
+			HashSalt = GenerateHashSalt();
 
-		ON_SCOPE_EXIT{ EVP_PKEY_free(PrivateKey); };
+			// Persist the salt since we'll persist the hash.
+			ensure(CredentialManager->SaveCredential(
+				PasswordSaltCredentialName,
+				FStringView(),
+				UTF8_TO_TCHAR(*HashSalt)));
+		}
 
-		X509* x509 = GenerateCertificate(PrivateKey);
-		if (!x509)
+		if (TSharedPtr<ICredential> SavedHash = CredentialManager->LoadCredential(PasswordHashCredentialName))
 		{
-			return TOptional<FString>{};
+			const int32 ExpectedHashStrLen = 2 * ScryptNumDerivedBytes;
+
+			if (ensure(SavedHash->GetBlob().Len() == ExpectedHashStrLen))
+			{
+				UE_LOGFMT(LogSwitchboard, Display, "Using stored password hash");
+				HashedAuthPassword.SetNumUninitialized(ScryptNumDerivedBytes);
+				UE::String::HexToBytes(SavedHash->GetBlob(), HashedAuthPassword.GetData());
+			}
+			else
+			{
+				UE_LOGFMT(LogSwitchboard, Error, "Stored password hash has wrong length (expected {ExpectedLen}, got {ActualLen}); discarding",
+					ExpectedHashStrLen, SavedHash->GetBlob().Len());
+			}
 		}
+	}
 
-		ON_SCOPE_EXIT{ X509_free(x509); };
+	// We also need to know if we should ignore JWTs issued prior to our last password change.
+	if (TSharedPtr<ICredential> SavedExpiry = CredentialManager->LoadCredential(LastPasswordChangeTimeCredentialName))
+	{
+		LexFromString(LastPasswordChangeTimeUtc, *FString(SavedExpiry->GetBlob()));
 
-		// Compute SHA256 fingerprint.
-		uint8 Digest[EVP_MAX_MD_SIZE];
-		uint32 DigestSize = 0;
-		check(1 == X509_digest(x509, EVP_sha256(), Digest, &DigestSize));
+		const FDateTime ExpiryTime = FDateTime::FromUnixTimestamp(LastPasswordChangeTimeUtc);
+		UE_LOGFMT(LogSwitchboard, Verbose, "Ignoring tokens issued prior to {ExpiryTime} UTC", ExpiryTime);
+	}
 
-		// Return the first half in the format "09:AB:CD:..."
-		FString Fingerprint;
-		for (uint32 ByteIdx = 0; ByteIdx < DigestSize/2; ++ByteIdx)
+	return true;
+}
+
+
+bool FSwitchboardAuthHelper::CheckOrCreateSelfSignedKeyPair(FSettings& InOutSettings)
+{
+	using namespace UE::SwitchboardListener;
+	using namespace UE::SwitchboardListener::Private;
+
+	const TTuple<FString, FString> DefaultCertKeyPaths = GetDefaultSelfSignedPaths();
+	const FString& CertificatePathStr = DefaultCertKeyPaths.Get<0>();
+	const FString& PrivateKeyPathStr = DefaultCertKeyPaths.Get<1>();
+	InOutSettings.Certificate.Emplace(TInPlaceType<FString>(), CertificatePathStr);
+	InOutSettings.PrivateKey.Emplace(TInPlaceType<FString>(), PrivateKeyPathStr);
+
+	bool bCanLoadKeyPair =
+		FPaths::FileExists(CertificatePathStr) &&
+		FPaths::FileExists(PrivateKeyPathStr);
+
+	// Platforms with a suitable credential store encrypt the generated
+	// private key with a random string, and persist it for the user.
+	if (CredentialManager->IsEncryptedAtRest())
+	{
+		// Try to load the private key password.
+		const FString PrivateKeyPwCredentialName =
+			FString::Printf(TEXT("PrivateKeyPassword_%s"), *PrivateKeyPathStr);
+
+		if (TSharedPtr<ICredential> PrivateKeyPwCredential =
+			CredentialManager->LoadCredential(PrivateKeyPwCredentialName))
 		{
-			ByteToHex(Digest[ByteIdx], Fingerprint);
-			if (ByteIdx < (DigestSize/2 - 1))
+			FStringView BlobView = PrivateKeyPwCredential->GetBlob();
+			auto BlobU8 = StringCast<UTF8CHAR>(BlobView.GetData(), BlobView.Len());
+			PrivateKeyPassword = (UTF8CHAR*)OPENSSL_secure_zalloc(BlobU8.Length() + 1);
+			FMemory::Memcpy(PrivateKeyPassword, BlobU8.Get(), BlobU8.Length());
+		}
+		else
+		{
+			if (bCanLoadKeyPair)
 			{
-				Fingerprint.AppendChar(TEXT(':'));
+				UE_LOGFMT(LogSwitchboard, Error, "Unable to load private key password; a new key pair will be generated");
+				bCanLoadKeyPair = false;
 			}
+
+			UE_LOGFMT(LogSwitchboard, Verbose, "Generating new random private key password");
+
+			constexpr int32 RandomPwByteLength = 20;
+
+			uint8* RandomPwByteBuf = (uint8*)OPENSSL_secure_malloc(RandomPwByteLength);
+			ON_SCOPE_EXIT{ OPENSSL_secure_clear_free(RandomPwByteBuf, RandomPwByteLength); };
+
+			TArrayView<uint8> RandomPwBytes(RandomPwByteBuf, RandomPwByteLength);
+			FillSecureRandom(RandomPwBytes);
+
+			PrivateKeyPassword = (UTF8CHAR*)OPENSSL_secure_zalloc((RandomPwByteLength * 2) + 1);
+			UE::String::BytesToHex(RandomPwBytes, PrivateKeyPassword);
+
+			ensure(CredentialManager->SaveCredential(
+				PrivateKeyPwCredentialName,
+				FStringView(),
+				UTF8_TO_TCHAR(PrivateKeyPassword)));
 		}
+	}
 
-		// Get PEM strings and write to files.
-		const FString PrivateKeyString = PrivateKeyToString(PrivateKey, PrivateKeyPassword);
-		const FString CertificateString = X509ToString(x509);
+	if (!bCanLoadKeyPair)
+	{
+		UE_LOGFMT(LogSwitchboard, Display, "SwitchboardListener requires a TLS certificate to be specified; generating self-signed certificate...");
 
-		TTuple<FString, FString> Paths = GetSelfSignedPaths();
-		const FString CertificateFile = Paths.Get<0>();
-		const FString PrivateKeyFile = Paths.Get<1>();
+		Tie(PrivateKey, Certificate) = CreateSelfSignedPair();
 
-		bool bSucceeded = true;
-		bSucceeded &= FFileHelper::SaveStringToFile(CertificateString, *CertificateFile);
-		bSucceeded &= FFileHelper::SaveStringToFile(PrivateKeyString, *PrivateKeyFile);
+		IPlatformFile& PlatformFile = FPlatformFileManager::Get().GetPlatformFile();
+		PlatformFile.CreateDirectoryTree(*GetSwitchboardDataDirectory());
 
-#if !PLATFORM_WINDOWS
-		EnsurePathOnlyOwnerReadable(PrivateKeyFile);
-#endif
+		if (!FFileHelper::SaveStringToFile(CertificateToString(Certificate), *CertificatePathStr))
+		{
+			UE_LOGFMT(LogSwitchboard, Fatal, "{Func}: Failed to save certificate to {CertPath}",
+				__FUNCTION__, CertificatePathStr);
+
+			return false;
+		}
+
+		if (!SavePrivateKeyToFile(PrivateKey, PrivateKeyPathStr, PrivateKeyPassword))
+		{
+			UE_LOGFMT(LogSwitchboard, Fatal, "{Func}: Failed to save private key to {KeyPath}",
+				__FUNCTION__, PrivateKeyPathStr);
+
+			return false;
+		}
+
+		if (!CheckCredentialPermissions(PrivateKeyPathStr))
+		{
+			UE_LOGFMT(LogSwitchboard, Error, "{Func}: Insecure permissions on private key file! ({Path})",
+				__FUNCTION__, PrivateKeyPathStr);
+		}
+
+		UE_LOGFMT(LogSwitchboard, Display, "Your new self-signed certificate fingerprint is: {Fingerprint}",
+			*GetX509TruncatedFingerprint(Certificate));
+	}
+
+	return true;
+}
+
+
+FUtf8String FSwitchboardAuthHelper::GenerateHashSalt()
+{
+	using namespace UE::SwitchboardListener;
+
+	UE_LOGFMT(LogSwitchboard, Verbose, "Generating new password hash salt");
+
+	constexpr int32 SaltByteLength = 4;
+	TStaticArray<uint8, SaltByteLength> SaltBytes;
+	FillSecureRandom(SaltBytes);
+
+	constexpr int32 SaltStrLength = SaltByteLength * 2 + 1;
+	TUtf8StringBuilder<SaltStrLength> SaltStrBuilder;
+	UE::String::BytesToHex(SaltBytes, SaltStrBuilder);
+
+	return SaltStrBuilder.ToString();
+}
+
+
+bool FSwitchboardAuthHelper::UpdateLastPasswordChangeTimeUtc(
+	FDateTime InChangeTimeUtc /* = FDateTime::UtcNow() */
+)
+{
+	using namespace UE::SwitchboardListener::Private;
+
+	LastPasswordChangeTimeUtc = InChangeTimeUtc.ToUnixTimestamp();
+	return CredentialManager->SaveCredential(LastPasswordChangeTimeCredentialName,
+		FStringView(),
+		LexToString(LastPasswordChangeTimeUtc)
+	);
+}
+
+
+void FSwitchboardAuthHelper::Shutdown()
+{
+	if (PrivateKey)
+	{
+		EVP_PKEY_free(PrivateKey);
+		PrivateKey = nullptr;
+	}
+
+	if (Certificate)
+	{
+		X509_free(Certificate);
+		Certificate = nullptr;
+	}
+
+	if (PrivateKeyPassword)
+	{
+		OPENSSL_secure_clear_free(PrivateKeyPassword, FCStringUtf8::Strlen(PrivateKeyPassword) + 1);
+		PrivateKeyPassword = nullptr;
+	}
+
+	if (PlainAuthPassword)
+	{
+		OPENSSL_secure_clear_free(PlainAuthPassword, FCStringUtf8::Strlen(PlainAuthPassword));
+		PlainAuthPassword = nullptr;
+	}
+
+	if (bSecureHeapNeedsShutdown)
+	{
+		CRYPTO_secure_malloc_done();
+		bSecureHeapNeedsShutdown = false;
+	}
+}
+
+
+const FUtf8String& FSwitchboardAuthHelper::GetCertificateFilePath() const
+{
+	return CertificateFilePath;
+}
+
+
+const FUtf8String& FSwitchboardAuthHelper::GetPrivateKeyFilePath() const
+{
+	return PrivateKeyFilePath;
+}
+
+
+const UTF8CHAR* FSwitchboardAuthHelper::GetPrivateKeyPassword() const
+{
+	return PrivateKeyPassword;
+}
+
+
+bool FSwitchboardAuthHelper::IsAuthPasswordSet() const
+{
+	return HashedAuthPassword.Num() > 0;
+}
+
+
+const UTF8CHAR* FSwitchboardAuthHelper::GetAuthPassword() const
+{
+	return PlainAuthPassword;
+}
+
+
+
+
+bool FSwitchboardAuthHelper::SetAuthPassword(
+	FStringView InNewPassword,
+	bool bExpireExistingTokens /* = true */
+)
+{
+	using namespace UE::SwitchboardListener;
+	using namespace UE::SwitchboardListener::Private;
+
+	if (!ensure(!InNewPassword.IsEmpty()))
+	{
+		return false;
+	}
+
+	auto InNewPasswordU8 = StringCast<UTF8CHAR>(InNewPassword.GetData(), InNewPassword.Len());
+	HashedAuthPassword.SetNumUninitialized(ScryptNumDerivedBytes);
+	check(DeriveScryptHash(InNewPasswordU8, HashSalt, HashedAuthPassword));
+
+	if (bExpireExistingTokens)
+	{
+		ensure(UpdateLastPasswordChangeTimeUtc());
+	}
+
+	if (CredentialManager->IsEncryptedAtRest())
+	{
+		if (PlainAuthPassword)
+		{
+			OPENSSL_secure_clear_free(PlainAuthPassword, FCStringUtf8::Strlen(PlainAuthPassword));
+		}
+
+		PlainAuthPassword = (UTF8CHAR*)OPENSSL_secure_zalloc(InNewPasswordU8.Length() + 1);
+		FMemory::Memcpy(PlainAuthPassword, InNewPasswordU8.Get(), InNewPasswordU8.Length());
+
+		UE_LOGFMT(LogSwitchboard, Display, "Storing provided password");
+		ensure(CredentialManager->SaveCredential(
+			PlainPasswordCredentialName,
+			FStringView(),
+			InNewPassword));
+	}
+	else
+	{
+		constexpr int32 HashStrLength = ScryptNumDerivedBytes * 2 + 1;
+		TStringBuilder<HashStrLength> SaltStrBuilder;
+		UE::String::BytesToHex(HashedAuthPassword, SaltStrBuilder);
+
+		UE_LOGFMT(LogSwitchboard, Display, "Storing hash of provided password");
+		ensure(CredentialManager->SaveCredential(
+			PasswordHashCredentialName,
+			FStringView(),
+			SaltStrBuilder.ToString()));
+	}
+
+	return true;
+}
+
+
+bool FSwitchboardAuthHelper::ValidatePassword(FStringView InAttempt)
+{
+	using namespace UE::SwitchboardListener::Private;
+
+	auto AttemptU8 = StringCast<UTF8CHAR>(InAttempt.GetData(), InAttempt.Len());
+
+	TArray<uint8> HashedAttempt;
+	HashedAttempt.SetNumUninitialized(ScryptNumDerivedBytes);
+	check(DeriveScryptHash(AttemptU8, HashSalt, HashedAttempt));
+
+	check(HashedAttempt.Num() == HashedAuthPassword.Num());
 
-		return bSucceeded ? Fingerprint : TOptional<FString>{};
+	return 0 == CRYPTO_memcmp(HashedAttempt.GetData(), HashedAuthPassword.GetData(), HashedAuthPassword.Num());
+}
+
+
+namespace UE::SwitchboardListener::Private
+{
+	FString MakeJWTHeader(FStringView InAlgorithm)
+	{
+		FString Header;
+		TSharedRef<TJsonWriter<TCHAR, TCondensedJsonPrintPolicy<TCHAR>>> JsonWriter =
+			TJsonWriterFactory<TCHAR, TCondensedJsonPrintPolicy<TCHAR>>::Create(&Header);
+
+		JsonWriter->WriteObjectStart();
+		JsonWriter->WriteValue(TEXTVIEW("alg"), InAlgorithm);
+		JsonWriter->WriteValue(TEXTVIEW("typ"), TEXTVIEW("JWT"));
+		JsonWriter->WriteObjectEnd();
+		ensure(JsonWriter->Close());
+
+		return Header;
+	}
+
+	FString MakeJWTPayload(const TMap<FStringView, FStringView>& InPayloadFields)
+	{
+		FString Payload;
+		TSharedRef<TJsonWriter<TCHAR, TCondensedJsonPrintPolicy<TCHAR>>> JsonWriter =
+			TJsonWriterFactory<TCHAR, TCondensedJsonPrintPolicy<TCHAR>>::Create(&Payload);
+
+		JsonWriter->WriteObjectStart();
+		for (const TPair<FStringView, FStringView>& Field : InPayloadFields)
+		{
+			JsonWriter->WriteValue(Field.Key, Field.Value);
+		}
+		JsonWriter->WriteObjectEnd();
+		ensure(JsonWriter->Close());
+
+		return Payload;
+	}
+}
+
+
+FString FSwitchboardAuthHelper::IssueJWT(const TMap<FStringView, FStringView>& InPayloadFields)
+{
+	using namespace UE::SwitchboardListener::Private;
+
+	// Make header
+	constexpr FStringView Header_Alg = TEXTVIEW("RS256");
+	const FString HeaderJson = MakeJWTHeader(Header_Alg);
+
+	// Make payload
+	TMap<FStringView, FStringView> PayloadFields = InPayloadFields;
+
+	constexpr FStringView Issuer_Field = TEXTVIEW("iss");
+	PayloadFields.Add(Issuer_Field, JwtIssuer_Value);
+
+	// Issued at the current time.
+	constexpr FStringView IssuedAt_Field = TEXTVIEW("iat");
+	const FDateTime IssuedAtTime = FDateTime::UtcNow();
+	const FString IssuedAtTimeStr = LexToString(IssuedAtTime.ToUnixTimestamp());
+	PayloadFields.Add(IssuedAt_Field, IssuedAtTimeStr);
+
+	// Expires in 10 years.
+	constexpr FStringView Expiration_Field = TEXTVIEW("exp");
+	const FDateTime ExpirationTime = IssuedAtTime + FTimespan::FromDays(365 * 10);
+	const FString ExpirationTimeStr = LexToString(ExpirationTime.ToUnixTimestamp());
+	PayloadFields.Add(Expiration_Field, ExpirationTimeStr);
+
+	const FString PayloadJson = MakeJWTPayload(PayloadFields);
+
+	// Encode and concatenate header + payload, giving us the message to be signed.
+	const FString HeaderB64 = FBase64::Encode(HeaderJson, EBase64Mode::UrlSafe).Replace(TEXT("="), TEXT(""));
+	const FString PayloadB64 = FBase64::Encode(PayloadJson, EBase64Mode::UrlSafe).Replace(TEXT("="), TEXT(""));
+	const FString Message = HeaderB64 + TEXT('.') + PayloadB64;
+
+	// Take a UTF8 view of the message for computing the signature.
+	auto MessageU8 = StringCast<UTF8CHAR>(*Message);
+	TArrayView<const uint8> MessageByteView((const uint8*)MessageU8.Get(), MessageU8.Length());
+
+	// Get the underlying RSA private key and cast it to the opaque handle type.
+	check(PrivateKey);
+	RSA* PrivateKeyRsa = EVP_PKEY_get0_RSA(PrivateKey);
+	check(PrivateKeyRsa);
+
+	FRSAKeyHandle KeyHandle = reinterpret_cast<FRSAKeyHandle>(PrivateKeyRsa);
+
+	// Compute the signed digest on the message hash.
+	TUniquePtr<FEncryptionContext> EncryptionContext = IPlatformCrypto::Get().CreateContext();
+	check(EncryptionContext);
+
+	TArray<uint8> HashedMessage;
+	check(EncryptionContext->CalcSHA256(MessageByteView, HashedMessage));
+
+	TArray<uint8> Signature;
+	check(EncryptionContext->DigestSign_RS256(HashedMessage, Signature, KeyHandle));
+
+	const FString SignatureB64 = FBase64::Encode(Signature, EBase64Mode::UrlSafe).Replace(TEXT("="), TEXT(""));
+	return Message + TEXT('.') + SignatureB64;
+}
+
+
+TOptional<FJsonWebToken> FSwitchboardAuthHelper::GetValidatedJWT(FStringView InEncodedJwt) const
+{
+	using namespace UE::SwitchboardListener::Private;
+
+	FJwtAlgorithm_RS256 Verifier;
+	if (!ensure(Verifier.SetPublicKey(CertPubKeyToString(Certificate))))
+	{
+		UE_LOGFMT(LogSwitchboard, Error, "Public key error during JWT validation");
+		return {};
+	}
+
+	TOptional<FJsonWebToken> Jwt = UE::JWT::FromString(InEncodedJwt);
+	if (!ensure(Jwt))
+	{
+		UE_LOGFMT(LogSwitchboard, Error, "Error parsing JWT");
+		return {};
 	}
-};
+
+	if (!ensure(Jwt->Verify(Verifier, JwtIssuer_Value)))
+	{
+		UE_LOGFMT(LogSwitchboard, Error, "JWT validation failed");
+		return {};
+	}
+
+	// Check to see if we've implicitly expired the JWT on password change.
+	int64 IssuedAt;
+	if (!ensure(Jwt->GetIssuedAt(IssuedAt)))
+	{
+		UE_LOGFMT(LogSwitchboard, Error, "Couldn't get JWT issued at time");
+		return {};
+	}
+
+	if (IssuedAt < LastPasswordChangeTimeUtc)
+	{
+		UE_LOGFMT(LogSwitchboard, Error, "JWT issued at time predates last password change");
+		return {};
+	}
+
+	return Jwt;
+}
+
+
+bool FSwitchboardAuthHelper::IsValidJWT(FStringView InEncodedJwt) const
+{
+	return GetValidatedJWT(InEncodedJwt).IsSet();
+}
diff --git a/Engine/Source/Programs/SwitchboardListener/SblCore/Private/SwitchboardAuth.h b/Engine/Source/Programs/SwitchboardListener/SblCore/Private/SwitchboardAuth.h
deleted file mode 100644
diff --git a/Engine/Source/Programs/SwitchboardListener/SblCore/Private/SwitchboardListener.cpp b/Engine/Source/Programs/SwitchboardListener/SblCore/Private/SwitchboardListener.cpp
--- a/Engine/Source/Programs/SwitchboardListener/SblCore/Private/SwitchboardListener.cpp
+++ b/Engine/Source/Programs/SwitchboardListener/SblCore/Private/SwitchboardListener.cpp
@@ -28,14 +28,8 @@
 #include "Misc/FileHelper.h"
 #include "Misc/Paths.h"
 #include "Misc/ScopeExit.h"
-#include "String/BytesToHex.h"
-#include "String/HexToBytes.h"
 #include <atomic>
 
-#if !PLATFORM_WINDOWS
-#include <string.h> // memset
-#endif
-
 
 #if PLATFORM_WINDOWS
 
@@ -130,16 +124,6 @@ namespace
 		Endpoint.Port = QuicAddrGetPort(&QuicAddr);
 		return Endpoint;
 	}
-
-	void SecureZero(void* Buffer, size_t Length)
-	{
-#if PLATFORM_WINDOWS
-		SecureZeroMemory(Buffer, Length);
-#else
-		static void* (* const volatile memset_ptr)(void*, int, size_t) = memset;
-		(memset_ptr)(Buffer, 0, Length);
-#endif
-	}
 }
 
 struct FRunningProcess
@@ -226,36 +210,6 @@ FSwitchboardCommandLineOptions FSwitchboardCommandLineOptions::FromString(const
 		}
 	}
 
-	OutOptions.SecureMode = ESecureMode::Unspecified;
-
-	if (SwitchPairs.Contains(TEXT("certfile")))
-	{
-		OutOptions.CertificateFile = SwitchPairs[TEXT("certfile")];
-		OutOptions.SecureMode = ESecureMode::CertificateFromFile;
-	}
-
-	if (SwitchPairs.Contains(TEXT("keyfile")))
-	{
-		OutOptions.PrivateKeyFile = SwitchPairs[TEXT("keyfile")];
-		OutOptions.SecureMode = ESecureMode::CertificateFromFile;
-	}
-
-	if (SwitchPairs.Contains(TEXT("certhash")))
-	{
-#if PLATFORM_WINDOWS
-		OutOptions.CertificateHash = SwitchPairs[TEXT("certhash")];
-
-		if (OutOptions.SecureMode == ESecureMode::CertificateFromFile)
-		{
-			UE_LOGFMT(LogSwitchboard, Error, "Both certificate file and hash arguments were specified; using certificate hash, IGNORING file");
-		}
-
-		OutOptions.SecureMode = ESecureMode::CertificateByHash;
-#else
-		UE_LOGFMT(LogSwitchboard, Fatal, "Certificate by hash is only supported on Windows");
-#endif
-	}
-
 	if (SwitchPairs.Contains(TEXT("port")))
 	{
 		OutOptions.Port = FCString::Atoi(*SwitchPairs[TEXT("port")]);
@@ -349,108 +303,42 @@ FSwitchboardListener::~FSwitchboardListener()
 {
 }
 
-bool FSwitchboardListener::Init(bool bStartListeningImmediately /* = true */)
+bool FSwitchboardListener::Init()
 {
-	if (Options.SecureMode == FSwitchboardCommandLineOptions::ESecureMode::Unspecified)
+	if (!AuthHelper.Initialize())
 	{
-		TTuple<FString, FString> CertKeyPaths = UE::SwitchboardListener::Certificates::GetSelfSignedPaths();
-
-		Options.SecureMode = FSwitchboardCommandLineOptions::ESecureMode::CertificateFromFile;
-		Options.CertificateFile = CertKeyPaths.Get<0>();
-		Options.PrivateKeyFile = CertKeyPaths.Get<1>();
-
-		// Try to load the private key password.
-		const FString PrivateKeyPwCredentialName = FString::Printf(
-			TEXT("PrivateKeyPassword_%s"), **Options.PrivateKeyFile);
-		if (TOptional<UE::SwitchboardListener::FCredential> PrivateKeyPwCredential =
-			UE::SwitchboardListener::LoadCredential(PrivateKeyPwCredentialName))
-		{
-			PrivateKeyPassword = PrivateKeyPwCredential->CredentialBlob;
-		}
-
-
-		if (!FPaths::FileExists(*Options.CertificateFile)
-			|| !FPaths::FileExists(*Options.PrivateKeyFile)
-			|| (!PrivateKeyPassword || PrivateKeyPassword->IsEmpty()))
-		{
-			UE_LOGFMT(LogSwitchboard, Display, "SwitchboardListener requires a TLS certificate to be specified; generating self-signed certificate...");
-
-			// Platforms with a persistent credential store encrypt the generated
-			// private key with a random string, and persist it for the user.
-			if (UE::SwitchboardListener::SupportsPersistentCredentials())
-			{
-				UE_LOGFMT(LogSwitchboard, Verbose, "Generating random private key password");
-				constexpr int32 RandomPwByteLength = 20;
-				TArray<uint8> RandomPwBytes;
-				RandomPwBytes.SetNumUninitialized(RandomPwByteLength);
-				UE::SwitchboardListener::FillSecureRandom(RandomPwBytes);
-				PrivateKeyPassword = BytesToHexLower(RandomPwBytes.GetData(), RandomPwByteLength);
-
-				UE::SwitchboardListener::FCredential PrivateKeyPwCredential = {
-					.CredentialName = PrivateKeyPwCredentialName,
-					.CredentialBlob = *PrivateKeyPassword,
-				};
-
-				check(UE::SwitchboardListener::SaveCredential(PrivateKeyPwCredential));
-			}
-
-			TOptional<FString> Fingerprint =
-				UE::SwitchboardListener::Certificates::CreateSelfSigned(
-					PrivateKeyPassword.Get(FString()));
-
-			check(Fingerprint.IsSet());
-
-			UE_LOGFMT(LogSwitchboard, Display, "Your new self-signed certificate fingerprint is: {Fingerprint}", *Fingerprint);
-		}
-	}
-
-	if (TOptional<UE::SwitchboardListener::FCredential> SavedCredential =
-		UE::SwitchboardListener::LoadCredential(PasswordCredentialName))
-	{
-		UE_LOGFMT(LogSwitchboard, Display, "Using stored password");
-		ExpectedAuthPassword = SavedCredential->CredentialBlob;
+		return false;
 	}
 
 	OnInitDelegate.Broadcast();
 
-	if (bStartListeningImmediately)
-	{
-		return StartListening();
-	}
-	else
-	{
-		return true;
-	}
+	return true;
 }
 
 
 void FSwitchboardListener::Shutdown()
 {
+	AuthHelper.Shutdown();
+
 	OnShutdownDelegate.Broadcast();
 }
 
 
-bool FSwitchboardListener::SetAuthPassword(const FString& NewPassword)
+const UTF8CHAR* FSwitchboardListener::GetAuthPassword() const
 {
-	if (NewPassword.IsEmpty())
-	{
-		return false;
-	}
+	return AuthHelper.GetAuthPassword();
+}
 
-	ExpectedAuthPassword = NewPassword;
 
-	if (UE::SwitchboardListener::SupportsPersistentCredentials())
-	{
-		UE::SwitchboardListener::FCredential NewCredential = {
-			.CredentialName = PasswordCredentialName,
-			.CredentialBlob = ExpectedAuthPassword,
-		};
+bool FSwitchboardListener::IsAuthPasswordSet() const
+{
+	return AuthHelper.IsAuthPasswordSet();
+}
 
-		UE_LOGFMT(LogSwitchboard, Display, "Storing provided password");
-		ensure(UE::SwitchboardListener::SaveCredential(NewCredential));
-	}
 
-	return true;
+bool FSwitchboardListener::SetAuthPassword(const FString& NewPassword)
+{
+	return AuthHelper.SetAuthPassword(NewPassword);
 }
 
 
@@ -493,7 +381,6 @@ bool FSwitchboardListener::StartListening()
 
 	QUIC_CERTIFICATE_FILE CertFile = { 0 };
 	QUIC_CERTIFICATE_FILE_PROTECTED CertFileProtected = { 0 };
-	QUIC_CERTIFICATE_HASH CertHash = { 0 };
 	QUIC_CREDENTIAL_CONFIG CredConfig = {
 		.Type = QUIC_CREDENTIAL_TYPE_NONE,
 		.Flags =
@@ -501,64 +388,22 @@ bool FSwitchboardListener::StartListening()
 			QUIC_CREDENTIAL_FLAG_USE_PORTABLE_CERTIFICATES,
 	};
 
-	// Path FString -> UTF8 backing stores for the ConfigurationLoadCredential call.
-	TArray<UTF8CHAR> CertificateFileArray, PrivateKeyFileArray, PrivateKeyPwArray;
-
-	if (Options.SecureMode == FSwitchboardCommandLineOptions::ESecureMode::CertificateFromFile)
+	if (AuthHelper.GetPrivateKeyPassword())
 	{
-		auto PrivateKeyFileU8 = StringCast<UTF8CHAR>(**Options.PrivateKeyFile);
-		PrivateKeyFileArray = TArray<UTF8CHAR>(PrivateKeyFileU8.Get(), PrivateKeyFileU8.Length() + 1);
-		auto CertificateFileU8 = StringCast<UTF8CHAR>(**Options.CertificateFile);
-		CertificateFileArray = TArray<UTF8CHAR>(CertificateFileU8.Get(), CertificateFileU8.Length() + 1);
-
-		if (PrivateKeyPassword)
-		{
-			auto PrivateKeyPwU8 = StringCast<UTF8CHAR>(**PrivateKeyPassword);
-			PrivateKeyPwArray = TArray<UTF8CHAR>(PrivateKeyPwU8.Get(), PrivateKeyPwU8.Length() + 1);
-			SecureZero(const_cast<UTF8CHAR*>(PrivateKeyPwU8.Get()), PrivateKeyPwU8.Length());
-
-			CertFileProtected.PrivateKeyFile = reinterpret_cast<const char*>(PrivateKeyFileArray.GetData());
-			CertFileProtected.CertificateFile = reinterpret_cast<const char*>(CertificateFileArray.GetData());
-			CertFileProtected.PrivateKeyPassword = reinterpret_cast<const char*>(PrivateKeyPwArray.GetData());
+		CertFileProtected.PrivateKeyFile = reinterpret_cast<const char*>(*AuthHelper.GetPrivateKeyFilePath());
+		CertFileProtected.CertificateFile = reinterpret_cast<const char*>(*AuthHelper.GetCertificateFilePath());
+		CertFileProtected.PrivateKeyPassword = reinterpret_cast<const char*>(AuthHelper.GetPrivateKeyPassword());
 
-			CredConfig.Type = QUIC_CREDENTIAL_TYPE_CERTIFICATE_FILE_PROTECTED;
-			CredConfig.CertificateFileProtected = &CertFileProtected;
-		}
-		else
-		{
-			CertFile.PrivateKeyFile = reinterpret_cast<const char*>(PrivateKeyFileArray.GetData());
-			CertFile.CertificateFile = reinterpret_cast<const char*>(CertificateFileArray.GetData());
-
-			CredConfig.Type = QUIC_CREDENTIAL_TYPE_CERTIFICATE_FILE;
-			CredConfig.CertificateFile = &CertFile;
-		}
-	}
-#if PLATFORM_WINDOWS
-	else if (Options.SecureMode == FSwitchboardCommandLineOptions::ESecureMode::CertificateByHash)
-	{
-		constexpr int32 ExpectedBytes = sizeof(CertHash.ShaHash);
-		constexpr int32 ExpectedStrLen = ExpectedBytes * 2;
-		const int32 HashStrLen = Options.CertificateHash->Len();
-		if (HashStrLen != ExpectedStrLen)
-		{
-			UE_LOGFMT(LogSwitchboard, Error, "Certificate hash should be {ExpectedStrLen} characters ({ExpectedBytes} bytes); got {HashStrLen} characters",
-				ExpectedStrLen, ExpectedBytes, HashStrLen);
-			return false;
-		}
-
-		const int32 HashLen = UE::String::HexToBytes(Options.CertificateHash.GetValue(), CertHash.ShaHash);
-		ensure(HashLen == ExpectedBytes);
-
-		UE_LOGFMT(LogSwitchboard, Display, "Using certificate with hash {CertificateHash}", Options.CertificateHash.GetValue());
-		CredConfig.Type = QUIC_CREDENTIAL_TYPE_CERTIFICATE_HASH;
-		CredConfig.CertificateHash = &CertHash;
+		CredConfig.Type = QUIC_CREDENTIAL_TYPE_CERTIFICATE_FILE_PROTECTED;
+		CredConfig.CertificateFileProtected = &CertFileProtected;
 	}
-#endif // PLATFORM_WINDOWS
 	else
 	{
-		// This should never happen.
-		checkf(false, TEXT("Unsupported credential configuration"));
-		return false;
+		CertFile.PrivateKeyFile = reinterpret_cast<const char*>(*AuthHelper.GetPrivateKeyFilePath());
+		CertFile.CertificateFile = reinterpret_cast<const char*>(*AuthHelper.GetCertificateFilePath());
+
+		CredConfig.Type = QUIC_CREDENTIAL_TYPE_CERTIFICATE_FILE;
+		CredConfig.CertificateFile = &CertFile;
 	}
 
 	QUIC_SETTINGS Settings = {};
@@ -587,20 +432,6 @@ bool FSwitchboardListener::StartListening()
 		UE_LOGFMT(LogSwitchboard, Error, "MsQuic ConfigurationLoadCredential failed with status {Status}", static_cast<int64>(Status));
 		return false;
 	}
-	else
-	{
-		if (PrivateKeyPwArray.Num())
-		{
-			SecureZero(PrivateKeyPwArray.GetData(), PrivateKeyPwArray.Num());
-		}
-
-		if (PrivateKeyPassword.IsSet())
-		{
-			SecureZero(PrivateKeyPassword->GetCharArray().GetData(),
-				PrivateKeyPassword->GetCharArray().Num() * sizeof(TCHAR));
-			PrivateKeyPassword.Reset();
-		}
-	}
 
 	// Create/allocate a new listener object.
 	if (QUIC_FAILED(Status = QuicApi->ListenerOpen(QuicRegistration, QuicListenerThunk, this, &QuicListener))) 
@@ -967,14 +798,34 @@ bool FSwitchboardListener::Task_Authenticate(const FSwitchboardAuthenticateTask&
 	UE::TUniqueLock<UE::FRecursiveMutex> ConnectionsScopeLock(ConnectionsLock);
 	FConnectionRef Connection = ConnectionsByEndpoint[InAuthTask.Recipient];
 
-	auto SendAuthResponse = [this, &InAuthTask, &Connection]() {
-		SendMessage(CreateMessage(
+	// If we decide later to issue a new JWT, we include it in SendAuthResponse.
+	TOptional<FString> IssueJwt;
+
+	auto SendAuthResponse = [this, &InAuthTask, &Connection, &IssueJwt]()
+		{
+			TMap<FString, FString> AuthResponse =
 			{
 				{ TEXT("command"), FSwitchboardAuthenticateTask::CommandName },
 				{ TEXT("id"), InAuthTask.TaskID.ToString() },
 				{ TEXT("bAuthenticated"), Connection->bAuthenticated ? TEXT("true") : TEXT("false") },
-			}),
-			InAuthTask.Recipient);
+			};
+
+			if (IssueJwt)
+			{
+				AuthResponse.Add(TEXT("jwt"), *IssueJwt);
+			}
+
+			SendMessage(CreateMessage(AuthResponse), InAuthTask.Recipient);
+
+			if (Connection->bAuthenticated)
+			{
+				// Send current state upon authentication
+				{
+					FSwitchboardStatePacket StatePacket;
+					FillStatePacket(StatePacket);
+					SendMessage(CreateMessage(StatePacket), InAuthTask.Recipient);
+				}
+			}
 		};
 
 	if (Connection->bAuthenticated)
@@ -992,21 +843,26 @@ bool FSwitchboardListener::Task_Authenticate(const FSwitchboardAuthenticateTask&
 		return false;
 	}
 
-	if (InAuthTask.Password == ExpectedAuthPassword)
+	if (InAuthTask.Jwt)
 	{
-		UE_LOGFMT(LogSwitchboard, Display, "Client {Endpoint} authenticated successfully", Connection->Endpoint.ToString());
-		Connection->bAuthenticated = true;
-		// TODO: Issue JWT?
-		SendAuthResponse();
-
-		// Send current state upon authentication
+		if (AuthHelper.IsValidJWT(*InAuthTask.Jwt))
+		{
+			UE_LOGFMT(LogSwitchboard, Display, "Client {Endpoint} authenticated successfully (using token)", Connection->Endpoint.ToString());
+			Connection->bAuthenticated = true;
+			SendAuthResponse();
+			return true;
+		}
+	}
+	else if (InAuthTask.Password)
+	{
+		if (AuthHelper.ValidatePassword(*InAuthTask.Password))
 		{
-			FSwitchboardStatePacket StatePacket;
-			FillStatePacket(StatePacket);
-			SendMessage(CreateMessage(StatePacket), InAuthTask.Recipient);
+			UE_LOGFMT(LogSwitchboard, Display, "Client {Endpoint} authenticated successfully (using password)", Connection->Endpoint.ToString());
+			Connection->bAuthenticated = true;
+			IssueJwt = AuthHelper.IssueJWT({});
+			SendAuthResponse();
+			return true;
 		}
-
-		return true;
 	}
 
 	UE_LOGFMT(LogSwitchboard, Warning, "Failed authentication attempt for client {Endpoint}", Connection->Endpoint.ToString());
diff --git a/Engine/Source/Programs/SwitchboardListener/SblCore/Private/SwitchboardMainCommon.cpp b/Engine/Source/Programs/SwitchboardListener/SblCore/Private/SwitchboardMainCommon.cpp
--- a/Engine/Source/Programs/SwitchboardListener/SblCore/Private/SwitchboardMainCommon.cpp
+++ b/Engine/Source/Programs/SwitchboardListener/SblCore/Private/SwitchboardMainCommon.cpp
@@ -29,6 +29,8 @@
 DEFINE_LOG_CATEGORY(LogSwitchboard);
 
 
+#if !defined(SWITCHBOARD_LISTENER_EXCLUDE_MAIN)
+
 int32 InitEngine()
 {
 	const TCHAR* CommandLine = FCommandLine::Get();
@@ -249,6 +251,10 @@ bool SwitchboardListenerMainShutdown()
 		GIOThreadPool->Destroy();
 	}
 
+#if STATS
+	FThreadStats::StopThread();
+#endif
+
 	FTaskGraphInterface::Shutdown();
 
 #if WITH_APPLICATION_CORE
@@ -316,3 +322,5 @@ int32 SwitchboardListenerMainWrapper()
 
 	return ExitCode;
 }
+
+#endif // #if !defined(SWITCHBOARD_LISTENER_EXCLUDE_MAIN)
diff --git a/Engine/Source/Programs/SwitchboardListener/SblCore/Private/SwitchboardProtocol.cpp b/Engine/Source/Programs/SwitchboardListener/SblCore/Private/SwitchboardProtocol.cpp
--- a/Engine/Source/Programs/SwitchboardListener/SblCore/Private/SwitchboardProtocol.cpp
+++ b/Engine/Source/Programs/SwitchboardListener/SblCore/Private/SwitchboardProtocol.cpp
@@ -196,15 +196,39 @@ FCreateTaskResult CreateTaskFromCommand(const FString& InCommand, const FIPv4End
 	}
 	else if (CommandName == FSwitchboardAuthenticateTask::CommandName)
 	{
-		TSharedPtr<FJsonValue> TokenField = TryGetCommandRequiredField(JsonData, TEXT("token"));
-		if (!TokenField)
+		// TODO: Delete me sometime after SBL 3.1. Not worth the breaking change.
+		const FStringView DeprecatedTokenFieldName = TEXTVIEW("token");
+		const bool bHasDeprecatedTokenField = JsonData->HasTypedField(DeprecatedTokenFieldName, EJson::String);
+
+		const FStringView JwtFieldName = TEXTVIEW("jwt");
+		const FStringView PasswordFieldName = TEXTVIEW("password");
+		const bool bHasJwtField = JsonData->HasTypedField(JwtFieldName, EJson::String);
+		const bool bHasPasswordField = JsonData->HasTypedField(PasswordFieldName, EJson::String);
+		if (!bHasDeprecatedTokenField && !bHasJwtField && !bHasPasswordField)
 		{
 			Result.Status = ECreateTaskStatus::Error_ParsingFailed;
 			return Result;
 		}
 
+		TOptional<FString> JwtField;
+		TOptional<FString> PasswordField;
+
+		if (bHasJwtField)
+		{
+			JwtField = JsonData->GetStringField(JwtFieldName);
+		}
+
+		if (bHasPasswordField)
+		{
+			PasswordField = JsonData->GetStringField(PasswordFieldName);
+		}
+		else if (bHasDeprecatedTokenField)
+		{
+			PasswordField = JsonData->GetStringField(DeprecatedTokenFieldName);
+		}
+
 		Result.Status = ECreateTaskStatus::Success;
-		Result.Task = MakeUnique<FSwitchboardAuthenticateTask>(MessageID, InEndpoint, TokenField->AsString());
+		Result.Task = MakeUnique<FSwitchboardAuthenticateTask>(MessageID, InEndpoint, JwtField, PasswordField);
 		return Result;
 	}
 
diff --git a/Engine/Source/Programs/SwitchboardListener/SblCore/Private/SwitchboardTasks.h b/Engine/Source/Programs/SwitchboardListener/SblCore/Private/SwitchboardTasks.h
--- a/Engine/Source/Programs/SwitchboardListener/SblCore/Private/SwitchboardTasks.h
+++ b/Engine/Source/Programs/SwitchboardListener/SblCore/Private/SwitchboardTasks.h
@@ -52,12 +52,21 @@ struct FSwitchboardTask
 
 struct FSwitchboardAuthenticateTask : public FSwitchboardTask
 {
-	FSwitchboardAuthenticateTask(const FGuid& InTaskId, const FIPv4Endpoint& InEndpoint, const FString& InPassword)
+	FSwitchboardAuthenticateTask(
+		const FGuid& InTaskId,
+		const FIPv4Endpoint& InEndpoint,
+		TOptional<FString> InJwt,
+		TOptional<FString> InPassword
+	)
 		: FSwitchboardTask{ ESwitchboardTaskType::Authenticate, InTaskId, InEndpoint }
+		, Jwt(InJwt)
 		, Password(InPassword)
-	{}
+	{
+		ensure(Jwt || Password);
+	}
 
-	FString Password;
+	TOptional<FString> Jwt;
+	TOptional<FString> Password;
 
 	//~ Begin FSwitchboardTask interface
 	static constexpr const TCHAR* CommandName = TEXT("authenticate");
diff --git a/Engine/Source/Programs/SwitchboardListener/SblCore/Private/Windows/WindowsSwitchboardCredential.cpp b/Engine/Source/Programs/SwitchboardListener/SblCore/Private/Windows/WindowsSwitchboardCredential.cpp
new file mode 100644
--- /dev/null
+++ b/Engine/Source/Programs/SwitchboardListener/SblCore/Private/Windows/WindowsSwitchboardCredential.cpp
@@ -0,0 +1,112 @@
+// Copyright Epic Games, Inc. All Rights Reserved.
+
+#include "Windows/WindowsSwitchboardCredential.h"
+#include "SwitchboardListenerApp.h" // for LogSwitchboard
+
+#include "Logging/StructuredLog.h"
+
+#include "Windows/AllowWindowsPlatformTypes.h"
+#include <wincred.h>
+#include "Windows/HideWindowsPlatformTypes.h"
+
+
+namespace UE::SwitchboardListener::Private
+{
+	class FWindowsCredential : public ICredential
+	{
+	public:
+		FWindowsCredential(CREDENTIALW* InNativeCredential)
+			: NativeCredential(InNativeCredential)
+		{
+			if (ensure(InNativeCredential))
+			{
+				UserView = FStringView(NativeCredential->UserName);
+				BlobView = FStringView(
+					reinterpret_cast<TCHAR*>(NativeCredential->CredentialBlob),
+					NativeCredential->CredentialBlobSize / sizeof(TCHAR));
+			}
+		}
+
+		virtual FStringView GetUser() override
+		{
+			return UserView;
+		}
+
+
+		virtual FStringView GetBlob() override
+		{
+			return BlobView;
+		}
+
+		virtual ~FWindowsCredential() override
+		{
+			if (ensure(NativeCredential))
+			{
+				::CredFree(NativeCredential);
+				NativeCredential = nullptr;
+			}
+		}
+
+	private:
+		CREDENTIALW* NativeCredential;
+		FStringView UserView;
+		FStringView BlobView;
+	};
+
+	FString GetCredentialTargetName(FStringView InCredentialName)
+	{
+		return FString::Printf(TEXT("SwitchboardListener_%*s"),
+			InCredentialName.Len(), InCredentialName.GetData());
+	}
+
+	TSharedPtr<ICredential> FWindowsCredentialManager::LoadCredential(FStringView InCredentialName)
+	{
+		const FString TargetName = GetCredentialTargetName(InCredentialName);
+
+		CREDENTIALW* OutCredPtr = nullptr;
+		if (::CredReadW(*TargetName, CRED_TYPE_GENERIC, 0, &OutCredPtr))
+		{
+			return MakeShared<FWindowsCredential>(OutCredPtr);
+		}
+		else
+		{
+			const int64 LastError = ::GetLastError();
+			if (LastError == ERROR_NOT_FOUND)
+			{
+				UE_LOGFMT(LogSwitchboard, Verbose, "CredReadW returned ERROR_NOT_FOUND");
+			}
+			else
+			{
+				UE_LOGFMT(LogSwitchboard, Warning, "CredReadW failed; GetLastError = {LastError}",
+					static_cast<int64>(::GetLastError()));
+			}
+		}
+
+		return {};
+	}
+
+	bool FWindowsCredentialManager::SaveCredential(FStringView InCredentialName, FStringView InUser, FStringView InBlob)
+	{
+		const FString TargetName = GetCredentialTargetName(InCredentialName);
+		const FString UserName(InUser);
+
+		CREDENTIALW Cred = { 0 };
+		Cred.Type = CRED_TYPE_GENERIC;
+		Cred.TargetName = const_cast<TCHAR*>(*TargetName);
+		Cred.CredentialBlobSize = sizeof(TCHAR) * InBlob.Len();
+		Cred.CredentialBlob = reinterpret_cast<BYTE*>(
+			const_cast<TCHAR*>(InBlob.GetData()));
+		Cred.Persist = CRED_PERSIST_LOCAL_MACHINE;
+		Cred.UserName = const_cast<TCHAR*>(*UserName);
+
+		if (!::CredWriteW(&Cred, 0))
+		{
+			UE_LOGFMT(LogSwitchboard, Warning, "CredWriteW failed; GetLastError = {LastError}",
+				static_cast<int64>(::GetLastError()));
+
+			return false;
+		}
+
+		return true;
+	}
+}; // namespace UE::SwitchboardListener::Private
diff --git a/Engine/Source/Programs/SwitchboardListener/SblCore/SblCore.Build.cs b/Engine/Source/Programs/SwitchboardListener/SblCore/SblCore.Build.cs
--- a/Engine/Source/Programs/SwitchboardListener/SblCore/SblCore.Build.cs
+++ b/Engine/Source/Programs/SwitchboardListener/SblCore/SblCore.Build.cs
@@ -6,6 +6,13 @@ public class SblCore : ModuleRules
 {
 	public SblCore(ReadOnlyTargetRules Target) : base(Target)
 	{
+		PublicIncludePathModuleNames.AddRange(
+			new string[]
+			{
+				"OpenSSL",
+			}
+		);
+
 		PublicDependencyModuleNames.AddRange(
 			new string[]
 			{
@@ -23,6 +30,7 @@ public class SblCore : ModuleRules
 			new string[] {
 				"Json",
 				"JsonUtilities",
+				"JWT",
 				"MsQuicRuntime",
 				"TraceLog",
 			}
diff --git a/Engine/Source/Programs/SwitchboardListener/SblSlate/Private/SblMainWindow.cpp b/Engine/Source/Programs/SwitchboardListener/SblSlate/Private/SblMainWindow.cpp
--- a/Engine/Source/Programs/SwitchboardListener/SblSlate/Private/SblMainWindow.cpp
+++ b/Engine/Source/Programs/SwitchboardListener/SblSlate/Private/SblMainWindow.cpp
@@ -130,7 +130,7 @@ void FSwitchboardListenerMainWindow::OnInit()
 
 	RootWindow->SetContent(CreateRootSwitcher());
 
-	if (Listener.GetAuthPassword().IsEmpty())
+	if (!Listener.IsAuthPasswordSet())
 	{
 		PanelSwitcher->SetActiveWidgetIndex(static_cast<int32>(EPanelIndices::PasswordPanel));
 	}
@@ -257,7 +257,11 @@ void FSwitchboardListenerMainWindow::CustomizeToolMenus_AddPasswordSection(UTool
 {
 	FToolMenuSection& PasswordSection = InMenu->AddSection("SBL_Password", LOCTEXT("SettingsMenu_PasswordSection_Label", "Password"), FToolMenuInsert("SBL_General", EToolMenuInsertType::After));
 
-	PasswordSection.AddMenuEntry("PasswordPanel", LOCTEXT("SettingsMenu_PasswordPanel_Label", "Show/Change Password"), FText(), FSlateIcon(),
+	const FText& PasswordEntryLabel = Listener.GetAuthPassword() != nullptr
+		? LOCTEXT("SettingsMenu_PasswordPanel_CanShow_Label", "Show/Change Password")
+		: LOCTEXT("SettingsMenu_PasswordPanel_NoShow_Label", "Change Password");
+
+	PasswordSection.AddMenuEntry("PasswordPanel", PasswordEntryLabel, FText(), FSlateIcon(),
 		FExecuteAction::CreateLambda([this]()
 			{
 				PasswordTextBox->SetText(FText::FromString(Listener.GetAuthPassword()));
@@ -398,9 +402,9 @@ TSharedRef<SWidget> FSwitchboardListenerMainWindow::CreateSetPasswordPanel()
 					SNew(SButton)
 					.Text_Lambda([this]()
 						{
-							return Listener.GetAuthPassword().IsEmpty()
-								? LOCTEXT("PasswordPanel_SetPasswordButton", "Set Password")
-								: LOCTEXT("PasswordPanel_ChangePasswordButton", "Change Password");
+							return Listener.IsAuthPasswordSet()
+								? LOCTEXT("PasswordPanel_ChangePasswordButton", "Change Password")
+								: LOCTEXT("PasswordPanel_SetPasswordButton", "Set Password");
 						})
 					.IsEnabled_Lambda([this]() { return !PasswordTextBox->GetText().IsEmpty(); })
 					.OnClicked_Raw(this, &FSwitchboardListenerMainWindow::OnSetPasswordClicked)
@@ -415,7 +419,7 @@ TSharedRef<SWidget> FSwitchboardListenerMainWindow::CreateSetPasswordPanel()
 				[
 					SNew(SButton)
 					.Text(LOCTEXT("PasswordPanel_CancelButton", "Cancel"))
-					.Visibility_Lambda([this]() { return Listener.GetAuthPassword().IsEmpty() ? EVisibility::Collapsed : EVisibility::Visible; })
+					.Visibility_Lambda([this]() { return Listener.IsAuthPasswordSet() ? EVisibility::Visible : EVisibility::Collapsed; })
 					.OnClicked_Lambda([this]()
 						{
 							PanelSwitcher->SetActiveWidgetIndex(static_cast<int32>(EPanelIndices::MainPanel));
diff --git a/Engine/Source/Programs/SwitchboardListener/SwitchboardListener/Private/SwitchboardListenerMain.cpp b/Engine/Source/Programs/SwitchboardListener/SwitchboardListener/Private/SwitchboardListenerMain.cpp
--- a/Engine/Source/Programs/SwitchboardListener/SwitchboardListener/Private/SwitchboardListenerMain.cpp
+++ b/Engine/Source/Programs/SwitchboardListener/SwitchboardListener/Private/SwitchboardListenerMain.cpp
@@ -31,8 +31,7 @@ int32 SwitchboardListenerMain()
 	FSwitchboardListener Listener(Options);
 	FSwitchboardListenerMainWindow MainWindow(Listener);
 
-	const bool bStartListeningImmediately = false;
-	if (!Listener.Init(bStartListeningImmediately))
+	if (!Listener.Init())
 	{
 		RequestEngineExit(TEXT("FSwitchboardListener::Init() failure"));
 		return 1;
diff --git a/Engine/Source/Programs/SwitchboardListener/SwitchboardListenerCmd/Private/SwitchboardListenerCmdMain.cpp b/Engine/Source/Programs/SwitchboardListener/SwitchboardListenerCmd/Private/SwitchboardListenerCmdMain.cpp
--- a/Engine/Source/Programs/SwitchboardListener/SwitchboardListenerCmd/Private/SwitchboardListenerCmdMain.cpp
+++ b/Engine/Source/Programs/SwitchboardListener/SwitchboardListenerCmd/Private/SwitchboardListenerCmdMain.cpp
@@ -2,6 +2,7 @@
 
 #include "SwitchboardListenerMainCommon.h"
 #include "SwitchboardListener.h"
+#include "SwitchboardAuth.h"
 #include "SwitchboardListenerVersion.h"
 
 #include "RequiredProgramMainCPPInclude.h"
@@ -60,10 +61,22 @@ int32 SwitchboardListenerMain()
 	}
 
 	FSwitchboardListener Listener(Options);
+	Listener.Init();
 
-	if (!Listener.Init())
+	if (!Listener.IsAuthPasswordSet())
 	{
-		RequestEngineExit(TEXT("FSwitchboardListener::Init() failure"));
+		GLog->Flush();
+
+		printf(
+			"\nPlease set a password for Switchboard Listener on this machine.\n"
+			"This password will be used to authenticate and establish a secure connection with Switchboard.\n");
+
+		Listener.SetAuthPassword(UE::SwitchboardListener::ReadPasswordFromStdin());
+	}
+
+	if (!Listener.StartListening())
+	{
+		RequestEngineExit(TEXT("FSwitchboardListener::StartListening() failure"));
 		return 1;
 	}
 
diff --git a/Engine/Source/Programs/SwitchboardListener/SwitchboardListenerTests.Target.cs b/Engine/Source/Programs/SwitchboardListener/SwitchboardListenerTests.Target.cs
new file mode 100644
--- /dev/null
+++ b/Engine/Source/Programs/SwitchboardListener/SwitchboardListenerTests.Target.cs
@@ -0,0 +1,19 @@
+// Copyright Epic Games, Inc. All Rights Reserved.
+
+using UnrealBuildTool;
+
+
+[SupportedPlatforms("Win64", "Linux")]
+public class SwitchboardListenerTestsTarget : TestTargetRules
+{
+	public SwitchboardListenerTestsTarget(TargetInfo Target) : base(Target)
+	{
+		SolutionDirectory = "Programs/Switchboard";
+		LaunchModuleName = "SwitchboardListenerTests";
+
+		// Don't pull in LaunchEngineLoop
+		GlobalDefinitions.Add("SWITCHBOARD_LISTENER_EXCLUDE_MAIN=1");
+
+		ExtraModuleNames.Add(LaunchModuleName);
+	}
+}
diff --git a/Engine/Source/Programs/SwitchboardListener/SwitchboardListenerTests/Private/SwitchboardAuthTests.cpp b/Engine/Source/Programs/SwitchboardListener/SwitchboardListenerTests/Private/SwitchboardAuthTests.cpp
new file mode 100644
--- /dev/null
+++ b/Engine/Source/Programs/SwitchboardListener/SwitchboardListenerTests/Private/SwitchboardAuthTests.cpp
@@ -0,0 +1,136 @@
+// Copyright Epic Games, Inc. All Rights Reserved.
+
+#include "TestHarness.h"
+
+#include "SwitchboardAuth.h"
+#include "SwitchboardCredentialInterface.h"
+
+#include "JsonWebToken.h"
+#include "HAL/PlatformProcess.h"
+#include "HAL/PlatformTime.h"
+#include "LowLevelTestsRunner/WarnFilterScope.h"
+
+
+namespace UE::SwitchboardListener::Private::Tests
+{
+	class FMockCredentialManager : public ICredentialManager
+	{
+		virtual TSharedPtr<ICredential> LoadCredential(FStringView InCredentialName) override
+		{
+			// no-op
+			return {};
+		}
+
+		virtual bool SaveCredential(FStringView InCredentialName, FStringView InUser, FStringView InBlob) override
+		{
+			// no-op
+			return true;
+		}
+	};
+
+	TEST_CASE("UE::SwitchboardListener::Auth")
+	{
+		EVP_PKEY* PrivateKey = nullptr;
+		X509* Certificate = nullptr;
+		Tie(PrivateKey, Certificate) = CreateSelfSignedPair();
+
+		SECTION("Self-signed key pair generation")
+		{
+			REQUIRE(PrivateKey != nullptr);
+			REQUIRE(Certificate != nullptr);
+		}
+
+		//////////////////////////////////////////////////////////////////////////
+
+		FSwitchboardAuthHelper::FSettings Settings;
+		Settings.Certificate.Emplace(TInPlaceType<X509*>(), Certificate);
+		Settings.PrivateKey.Emplace(TInPlaceType<EVP_PKEY*>(), PrivateKey);
+
+		FMockCredentialManager MockCredMgr;
+		Settings.CredentialManager = &MockCredMgr;
+
+		SECTION("Password hashing")
+		{
+			constexpr FStringView TestPassword = TEXTVIEW("ExamplePassword");
+			constexpr FStringView IncorrectPassword = TEXTVIEW("WrongPassword");
+
+			FSwitchboardAuthHelper AuthHelper;
+			CHECK(AuthHelper.Initialize(Settings));
+
+			CHECK(!AuthHelper.IsAuthPasswordSet());
+
+			CHECK(AuthHelper.SetAuthPassword(TestPassword));
+
+			CHECK(AuthHelper.IsAuthPasswordSet());
+
+			CHECK(!AuthHelper.ValidatePassword(IncorrectPassword));
+
+			CHECK(AuthHelper.ValidatePassword(TestPassword));
+		}
+
+		//////////////////////////////////////////////////////////////////////////
+
+		SECTION("JWT")
+		{
+			FSwitchboardAuthHelper AuthHelper;
+			CHECK(AuthHelper.Initialize(Settings));
+
+			constexpr FStringView TestClaimName = TEXTVIEW("MyTestJwtField");
+			constexpr FStringView TestClaimValue = TEXTVIEW("MyTestJwtValue");
+
+			const FString JwtString = AuthHelper.IssueJWT({
+				{ TestClaimName, TestClaimValue }
+			});
+
+			const double JwtIssuePlatformTime = FPlatformTime::Seconds();
+
+			CHECK(!JwtString.IsEmpty());
+
+			TOptional<FJsonWebToken> ValidatedJwt = AuthHelper.GetValidatedJWT(JwtString);
+			CHECK(ValidatedJwt.IsSet());
+
+			if (ValidatedJwt)
+			{
+				TSharedPtr<FJsonValue> ValidatedPayloadField =
+					ValidatedJwt->GetClaim<EJson::String>(TestClaimName);
+				CHECK(ValidatedPayloadField);
+				if (ValidatedPayloadField)
+				{
+					FString ValidatedClaimValue;
+					const bool bGotClaimString = ValidatedPayloadField->TryGetString(ValidatedClaimValue);
+					CHECK(bGotClaimString);
+					if (bGotClaimString)
+					{
+						CHECK_EQUAL(ValidatedClaimValue, TestClaimValue);
+					}
+				}
+			}
+
+			// Test expiry. Initially, the JWT is valid.
+			CHECK(AuthHelper.IsValidJWT(JwtString));
+
+			// We need at least one second to elapse after issuing the JWT.
+			const double SecondsSinceIssuance = FPlatformTime::Seconds() - JwtIssuePlatformTime;
+			FPlatformProcess::Sleep(1.0f - SecondsSinceIssuance);
+
+			const bool bExpireExistingTokens_true = true;
+			CHECK(AuthHelper.SetAuthPassword(TEXTVIEW("DifferentPassword"), bExpireExistingTokens_true));
+
+			// Previously issued JWT was expired by a password change.
+			{
+				UE::Testing::FWarnFilterScope _([](const TCHAR* Message, ELogVerbosity::Type Verbosity, const FName& Category)
+					{
+						constexpr FStringView ExpectedError = TEXTVIEW("JWT issued at time predates last password change");
+						const bool bFiltered =
+							0 == FCString::Strncmp(Message, ExpectedError.GetData(), ExpectedError.Len())
+							&& Verbosity == ELogVerbosity::Type::Error
+							&& Category == TEXT("LogSwitchboard");
+
+						return bFiltered;
+					});
+
+				CHECK(!AuthHelper.IsValidJWT(JwtString));
+			}
+		}
+	}
+} // namespace UE::SwitchboardListener::Private::Tests
diff --git a/Engine/Source/Programs/SwitchboardListener/SwitchboardListenerTests/Private/SwitchboardCredentialTests.cpp b/Engine/Source/Programs/SwitchboardListener/SwitchboardListenerTests/Private/SwitchboardCredentialTests.cpp
new file mode 100644
--- /dev/null
+++ b/Engine/Source/Programs/SwitchboardListener/SwitchboardListenerTests/Private/SwitchboardCredentialTests.cpp
@@ -0,0 +1,37 @@
+// Copyright Epic Games, Inc. All Rights Reserved.
+
+#include "TestHarness.h"
+
+#include "SwitchboardAuth.h"
+#include "SwitchboardCredential.h"
+
+#include "HAL/PlatformTime.h"
+#include "Math/UnrealMathUtility.h"
+
+
+namespace UE::SwitchboardListener::Private::Tests
+{
+	TEST_CASE("UE::SwitchboardListener::Credential")
+	{
+		FMath::RandInit(FPlatformTime::Cycles());
+		FMath::SRandInit(FPlatformTime::Cycles());
+
+		SECTION("Basic save and load")
+		{
+			constexpr FStringView TestCredentialName = TEXTVIEW("TestCredential");
+			const FString TestRandomUser = FString::FromInt(FMath::RandRange(int32(0), MAX_int32-1));
+			const FString TestRandomBlob = FString::FromInt(FMath::RandRange(int32(0), MAX_int32-1));
+
+			ICredentialManager& CredMgr = ICredentialManager::GetPlatformCredentialManager();
+			CHECK(CredMgr.SaveCredential(TestCredentialName, TestRandomUser, TestRandomBlob));
+
+			TSharedPtr<ICredential> LoadedCredential = CredMgr.LoadCredential(TestCredentialName);
+			CHECK(LoadedCredential);
+			if (LoadedCredential)
+			{
+				CHECK_EQUAL(LoadedCredential->GetUser(), TestRandomUser);
+				CHECK_EQUAL(LoadedCredential->GetBlob(), TestRandomBlob);
+			}
+		}
+	}
+} // namespace UE::SwitchboardListener::Private::Tests
diff --git a/Engine/Source/Programs/SwitchboardListener/SwitchboardListenerTests/SwitchboardListenerTests.Build.cs b/Engine/Source/Programs/SwitchboardListener/SwitchboardListenerTests/SwitchboardListenerTests.Build.cs
new file mode 100644
--- /dev/null
+++ b/Engine/Source/Programs/SwitchboardListener/SwitchboardListenerTests/SwitchboardListenerTests.Build.cs
@@ -0,0 +1,20 @@
+// Copyright Epic Games, Inc. All Rights Reserved.
+
+using UnrealBuildTool;
+
+
+public class SwitchboardListenerTests : TestModuleRules
+{
+	public SwitchboardListenerTests(ReadOnlyTargetRules Target) : base(Target)
+	{
+		PrivateDependencyModuleNames.AddRange(new string[] {
+			"SblCore",
+
+			"JWT",
+		});
+
+		PrivateIncludePathModuleNames.Add("LowLevelTestsRunner");
+
+		UpdateBuildGraphPropertiesFile(new Metadata() { TestName = "SwitchboardListener", TestShortName = "SwitchboardListener" });
+	}
+}
