SwitchboardListener: Fix crash when saved password credential does not exist - regression from Linux credential implementation.

Ensure that we always generate the ephemeral scrypt salt.

#rnx

diff --git a/Engine/Source/Programs/SwitchboardListener/SblCore/Internal/SwitchboardListenerVersion.h b/Engine/Source/Programs/SwitchboardListener/SblCore/Internal/SwitchboardListenerVersion.h
--- a/Engine/Source/Programs/SwitchboardListener/SblCore/Internal/SwitchboardListenerVersion.h
+++ b/Engine/Source/Programs/SwitchboardListener/SblCore/Internal/SwitchboardListenerVersion.h
@@ -4,4 +4,4 @@
 
 #define SBLISTENER_VERSION_MAJOR 3
 #define SBLISTENER_VERSION_MINOR 0
-#define SBLISTENER_VERSION_PATCH 6
+#define SBLISTENER_VERSION_PATCH 7
