Convert MsQuic error codes into human readable form.

diff --git a/Engine/Source/Programs/SwitchboardListener/SblCore/Private/SwitchboardListener.cpp b/Engine/Source/Programs/SwitchboardListener/SblCore/Private/SwitchboardListener.cpp
--- a/Engine/Source/Programs/SwitchboardListener/SblCore/Private/SwitchboardListener.cpp
+++ b/Engine/Source/Programs/SwitchboardListener/SblCore/Private/SwitchboardListener.cpp
@@ -126,6 +126,40 @@ namespace
 		Endpoint.Port = QuicAddrGetPort(&QuicAddr);
 		return Endpoint;
 	}
+
+	FString GetQuicErrorStr(QUIC_STATUS InStatus)
+	{
+		static TMap<QUIC_STATUS, FString> StatStringMap = {
+			{QUIC_STATUS_SUCCESS			 , TEXT("The operation completed successfully.")},
+			{QUIC_STATUS_PENDING			 , TEXT("The operation is pending.")},
+			{QUIC_STATUS_CONTINUE			 , TEXT("The operation will continue.")},
+			{QUIC_STATUS_OUT_OF_MEMORY		 , TEXT("Allocation of memory failed.")},
+			{QUIC_STATUS_INVALID_PARAMETER	 , TEXT("An invalid parameter was encountered.")},
+			{QUIC_STATUS_INVALID_STATE		 , TEXT("The current state was not valid for this operation.")},
+			{QUIC_STATUS_NOT_SUPPORTED	     , TEXT("The operation was not supported.")},
+			{QUIC_STATUS_NOT_FOUND	         , TEXT("The object was not found.")},
+			{QUIC_STATUS_BUFFER_TOO_SMALL	 , TEXT("The buffer was too small for the operation.")},
+			{QUIC_STATUS_HANDSHAKE_FAILURE	 , TEXT("The connection handshake failed.")},
+			{QUIC_STATUS_ABORTED			 , TEXT("The connection or stream was aborted.")},
+			{QUIC_STATUS_ADDRESS_IN_USE		 , TEXT("The local address is already in use.")},
+			{QUIC_STATUS_INVALID_ADDRESS	 , TEXT("Binding to socket failed, likely caused by a family mismatch between local and remote address.")},
+			{QUIC_STATUS_CONNECTION_TIMEOUT	 , TEXT("The connection timed out waiting for a response from the peer.")},
+			{QUIC_STATUS_CONNECTION_IDLE	 , TEXT("The connection timed out from inactivity.")},
+			{QUIC_STATUS_INTERNAL_ERROR		 , TEXT("An internal error was encountered.")},
+			{QUIC_STATUS_UNREACHABLE		 , TEXT("The server is currently unreachable.")},
+			{QUIC_STATUS_CONNECTION_REFUSED	 , TEXT("The server refused the connection.")},
+			{QUIC_STATUS_PROTOCOL_ERROR		 , TEXT("A protocol error was encountered.")},
+			{QUIC_STATUS_VER_NEG_ERROR		 , TEXT("A version negotiation error was encountered.")},
+			{QUIC_STATUS_USER_CANCELED		 , TEXT("The peer app/user canceled the connection during the handshake.")},
+			{QUIC_STATUS_ALPN_NEG_FAILURE	 , TEXT("The connection handshake failed to negotiate a common ALPN.")},
+			{QUIC_STATUS_STREAM_LIMIT_REACHED , TEXT("A stream failed to start because the peer doesn't allow any more to be open at this time.")}
+		};
+		if (FString* ErrorString = StatStringMap.Find(InStatus))
+		{
+			return *ErrorString;
+		}
+		return TEXT("Unknown MsQuic error.");
+	}
 }
 
 struct FRunningProcess
@@ -446,7 +480,7 @@ bool FSwitchboardListener::StartListening()
 	QUIC_STATUS Status;
 	if (QUIC_FAILED(Status = MsQuicOpen2(&QuicApi)))
 	{
-		UE_LOGFMT(LogSwitchboard, Error, "MsQuicOpen2 failed with status {Status}", static_cast<int64>(Status));
+		UE_LOGFMT(LogSwitchboard, Error, "MsQuicOpen2 failed with status: {Status}", *GetQuicErrorStr(Status));
 		return false;
 	}
 
@@ -455,7 +489,7 @@ bool FSwitchboardListener::StartListening()
 	// the execution profile, using the default "low latency" profile.
 	const QUIC_REGISTRATION_CONFIG RegConfig = { "switchboardlistener", QUIC_EXECUTION_PROFILE_LOW_LATENCY };
 	if (QUIC_FAILED(Status = QuicApi->RegistrationOpen(&RegConfig, &QuicRegistration))) {
-		UE_LOGFMT(LogSwitchboard, Error, "MsQuic RegistrationOpen failed with status {Status}", static_cast<int64>(Status));
+		UE_LOGFMT(LogSwitchboard, Error, "MsQuic RegistrationOpen failed with status: {Status}", *GetQuicErrorStr(Status));
 		return false;
 	}
 
@@ -502,21 +536,21 @@ bool FSwitchboardListener::StartListening()
 	// Allocate/initialize the configuration object with the configured ALPN and settings.
 	if (QUIC_FAILED(Status = QuicApi->ConfigurationOpen(QuicRegistration, &SblAlpn, 1, &Settings, sizeof(Settings), NULL, &QuicConfiguration)))
 	{
-		UE_LOGFMT(LogSwitchboard, Error, "MsQuic ConfigurationOpen failed with status {Status}", static_cast<int64>(Status));
+		UE_LOGFMT(LogSwitchboard, Error, "MsQuic ConfigurationOpen failed with status: {Status}", *GetQuicErrorStr(Status));
 		return false;
 	}
 
 	// Loads the TLS credential part of the configuration.
 	if (QUIC_FAILED(Status = QuicApi->ConfigurationLoadCredential(QuicConfiguration, &CredConfig)))
 	{
-		UE_LOGFMT(LogSwitchboard, Error, "MsQuic ConfigurationLoadCredential failed with status {Status}", static_cast<int64>(Status));
+		UE_LOGFMT(LogSwitchboard, Error, "MsQuic ConfigurationLoadCredential failed with status: {Status}", *GetQuicErrorStr(Status));
 		return false;
 	}
 
 	// Create/allocate a new listener object.
 	if (QUIC_FAILED(Status = QuicApi->ListenerOpen(QuicRegistration, QuicListenerThunk, this, &QuicListener))) 
 	{
-		UE_LOGFMT(LogSwitchboard, Error, "MsQuic ListenerOpen failed with status {Status}", static_cast<int64>(Status));
+		UE_LOGFMT(LogSwitchboard, Error, "MsQuic ListenerOpen failed with status: {Status}", *GetQuicErrorStr(Status));
 		return false;
 	}
 
@@ -524,7 +558,7 @@ bool FSwitchboardListener::StartListening()
 	QUIC_ADDR QuicAddr = QuicAddrFromEndpoint(*ListenerEndpoint);
 	if (QUIC_FAILED(Status = QuicApi->ListenerStart(QuicListener, &SblAlpn, 1, &QuicAddr))) 
 	{
-		UE_LOGFMT(LogSwitchboard, Error, "MsQuic ListenerStart failed with status {Status}", static_cast<int64>(Status));
+		UE_LOGFMT(LogSwitchboard, Error, "MsQuic ListenerStart failed with status: {Status}", *GetQuicErrorStr(Status));
 		return false;
 	}
 
