Switchboard: Updated the fix to the hitches issue because the other gpu stats api is exhibiting persistent and periodic micro hitches.

Added a "GPU Stats" checkbox to nDisplay monitor. Defaults to off, and should be kept off unless the user wants to monitor the GPU stats. The reason to recommend it to be off by default is that these queries can GPU in a state that makes it more prone to hitches.

We now close NvApi when no new usages have happened for the last 10 minutes.

Fixed a missing signal connection for on listener connected in Unreal device which was not allowing a full ndisplay monitor table refresh for the node that gets connected.

The gpu stats will now read as n/a when they are not being polled. This is because having stale data of these can be misleading to the user.

Also changed order of columns in the monitor because the OS is very wide and not as dynamic as the others.

Increased window initial size from 1192 to 1280 to accomodate the new GPU Stats checkbox.

"Fix ExeFlags" button removed, now the command is sent every time the program starts signal happens for an nDisplay device.
It is still there but not shown, and internally renamed to "Disable FSO" (and column to "FSO"), because it is the only thing we actually do.

diff --git a/Engine/Plugins/VirtualProduction/Switchboard/Source/Switchboard/switchboard/devices/ndisplay/ndisplay_monitor.py b/Engine/Plugins/VirtualProduction/Switchboard/Source/Switchboard/switchboard/devices/ndisplay/ndisplay_monitor.py
--- a/Engine/Plugins/VirtualProduction/Switchboard/Source/Switchboard/switchboard/devices/ndisplay/ndisplay_monitor.py
+++ b/Engine/Plugins/VirtualProduction/Switchboard/Source/Switchboard/switchboard/devices/ndisplay/ndisplay_monitor.py
@@ -33,7 +33,10 @@ class nDisplayMonitor(QAbstractTableModel):
 
     # If false, the button to disable full screen optimizations is hidden
     # and periodic polling of the state of this feature is disabled.
-    use_exe_flags = True
+    show_disable_fso_btn = False
+
+    # Determines if GPU stats are regularly queried or not.
+    poll_gpu_stats = False
 
     def __init__(self, parent):
         QAbstractTableModel.__init__(self, parent)
@@ -69,12 +72,6 @@ class nDisplayMonitor(QAbstractTableModel):
             'InFocus':
                 'Whether nDisplay instance window is in focus. It is '
                 'recommended to be in focus.',
-            'ExeFlags':
-                'It is recommended to disable fullscreen optimizations in the\n'
-                'unreal executable because it has been associated with tearing.\n'
-                'Only available once the render node process is running. \n'
-                '(Expected value is "DISABLEDXMAXIMIZEDWINDOWEDMODE")',
-            'OsVer': 'Operating system version',
             'CpuUtilization':
                 'CPU utilization average. The number of overloaded cores (>'
                 f'{self.CORE_OVERLOAD_THRESH}% load) will be displayed in '
@@ -89,11 +86,14 @@ class nDisplayMonitor(QAbstractTableModel):
             'GpuTemperature':
                 'GPU temperature in degrees celsius. (Max across all '
                 'sensors.)',
+            'FSO':
+                'It is recommended to disable Full Screen Optimizations in the\n'
+                'unreal executable because it has been associated with tearing.\n'
+                'Only available once the render node process is running. \n'
+                '(Expected value is "no")',
+            'OsVer': 'Operating system version',
         }
 
-        if not self.use_exe_flags:
-            HEADER_DATA.pop('ExeFlags')
-
         self.colnames = list(HEADER_DATA.keys())
         self.tooltips = list(HEADER_DATA.values())
 
@@ -125,8 +125,8 @@ class nDisplayMonitor(QAbstractTableModel):
             is_good = (not is_program_running) or ('yes' in value)
             return self.COLOR_NORMAL if is_good else self.COLOR_WARNING
 
-        if colname == 'ExeFlags':
-            is_good = (not is_program_running) or ('DISABLEDXMAXIMIZEDWINDOWEDMODE' in value)
+        if colname == 'FSO':
+            is_good = (not is_program_running) or ('no' in value)
             return self.COLOR_NORMAL if is_good else self.COLOR_WARNING
 
         if colname == 'Displays':
@@ -336,8 +336,10 @@ class nDisplayMonitor(QAbstractTableModel):
         request_flags &= ~SyncStatusRequestFlags.MosaicTopos
         request_flags &= ~SyncStatusRequestFlags.DriverInfo
 
-        if not self.use_exe_flags:
-            request_flags &= ~SyncStatusRequestFlags.ProgramLayers
+        if not self.poll_gpu_stats:
+            request_flags &= ~SyncStatusRequestFlags.GpuCoreClockKhz
+            request_flags &= ~SyncStatusRequestFlags.GpuTemperature
+            request_flags &= ~SyncStatusRequestFlags.GpuUtilization
 
         self.poll_sync_status(request_flags)
 
@@ -501,8 +503,13 @@ class nDisplayMonitor(QAbstractTableModel):
 
         # Show Exe flags (like Disable Fullscreen Optimization)
         if SyncStatusRequestFlags.ProgramLayers in request_flags:
-            data['ExeFlags'] = '\n'.join([
-                layer for layer in sync_status['programLayers'][1:]])
+            layers = '\n'.join([layer for layer in sync_status['programLayers'][1:]])
+            if 'DISABLEDXMAXIMIZEDWINDOWEDMODE' in layers:
+                data['FSO'] = 'no'
+            elif len(device.program_start_queue.running_programs_named('unreal')):
+                data['FSO'] = 'yes'
+            else:
+                data['FSO'] = 'n/a'
 
         # Driver version
         if SyncStatusRequestFlags.DriverInfo in request_flags:
@@ -540,7 +547,7 @@ class nDisplayMonitor(QAbstractTableModel):
             except (KeyError, ValueError):
                 data['CpuUtilization'] = self.DATA_MISSING
 
-        # Memory utilization
+        # Host memory utilization
         if SyncStatusRequestFlags.AvailablePhysicalMemory in request_flags:
             try:
                 gb = 1024 * 1024 * 1024
@@ -552,33 +559,40 @@ class nDisplayMonitor(QAbstractTableModel):
             except TypeError:
                 data['MemUtilization'] = self.DATA_MISSING
 
-        # GPU utilization + clocks
-        if (SyncStatusRequestFlags.GpuUtilization | SyncStatusRequestFlags.GpuCoreClockKhz) in request_flags:
-            try:
-                gpu_stats = list(map(
-                    lambda x: f"#{x[0]}: {x[1]:.0f}% ({x[2] / 1000:.0f} MHz)",
-                    zip(count(), sync_status['gpuUtilization'],
-                        sync_status['gpuCoreClocksKhz'])))
-
-                if len(gpu_stats) > 0:
-                    data['GpuUtilization'] = '\n'.join(gpu_stats)
-                else:
+        if self.poll_gpu_stats:
+
+            # GPU utilization + clocks
+            if (SyncStatusRequestFlags.GpuUtilization | SyncStatusRequestFlags.GpuCoreClockKhz) in request_flags:
+                try:
+                    gpu_stats = list(map(
+                        lambda x: f"#{x[0]}: {x[1]:.0f}% ({x[2] / 1000:.0f} MHz)",
+                        zip(count(), sync_status['gpuUtilization'],
+                            sync_status['gpuCoreClocksKhz'])))
+
+                    if len(gpu_stats) > 0:
+                        data['GpuUtilization'] = '\n'.join(gpu_stats)
+                    else:
+                        data['GpuUtilization'] = self.DATA_MISSING
+                except (KeyError, TypeError):
                     data['GpuUtilization'] = self.DATA_MISSING
-            except (KeyError, TypeError):
-                data['GpuUtilization'] = self.DATA_MISSING
-
-        # GPU temperature
-        if SyncStatusRequestFlags.GpuTemperature in request_flags:
-            try:
-                temps = [t if t != -2147483648 else self.DATA_MISSING for t in sync_status['gpuTemperature']]
 
-                if len(temps) > 0:
-                    data['GpuTemperature'] = '\n'.join(
-                        map(lambda x: f"#{x[0]}: {x[1]}° C", zip(count(), temps)))
-                else:
+            # GPU temperature
+            if SyncStatusRequestFlags.GpuTemperature in request_flags:
+                try:
+                    temps = [t if t != -2147483648 else self.DATA_MISSING for t in sync_status['gpuTemperature']]
+
+                    if len(temps) > 0:
+                        data['GpuTemperature'] = '\n'.join(
+                            map(lambda x: f"#{x[0]}: {x[1]}° C", zip(count(), temps)))
+                    else:
+                        data['GpuTemperature'] = self.DATA_MISSING
+                except (KeyError, TypeError):
                     data['GpuTemperature'] = self.DATA_MISSING
-            except (KeyError, TypeError):
-                data['GpuTemperature'] = self.DATA_MISSING
+        else:
+            # If we're generally not polling stats, show n/a as to not be misleading
+            # since it is such a volatile metric.
+            data['GpuTemperature'] = 'n/a'
+            data['GpuUtilization'] = 'n/a'
 
     def on_get_sync_status(self, device, message):
         '''
@@ -741,16 +755,15 @@ class nDisplayMonitor(QAbstractTableModel):
         self.poll_sync_status(SyncStatusRequestFlags.all())
 
     @QtCore.Slot()
-    def on_fix_exe_flags_clicked(self):
+    def on_disable_fso_clicked(self):
         ''' Tries to force the correct UnrealEditor.exe flags '''
         for devicedata in self.devicedatas.values():
             device = devicedata['device']
             data = devicedata['data']
 
-            good_string = 'DISABLEDXMAXIMIZEDWINDOWEDMODE'
-
-            if good_string not in data['ExeFlags']:
-                device.fix_exe_flags()
+            # Generally, want FSO to be disabled.
+            if data['FSO'] != 'no':
+                device.disable_fso()
 
     @QtCore.Slot()
     def on_soft_kill_clicked(self):
@@ -770,3 +783,11 @@ class nDisplayMonitor(QAbstractTableModel):
         for devicedata in self.devicedatas.values():
             device = devicedata['device']
             device.minimize_windows()
+
+    @QtCore.Slot()
+    def on_gpu_stats_toggled(self, state: int):
+        ''' Called when the GPU Stats checkbox is toggled. We update the state variable accordingly. '''
+        if state == Qt.Checked.value:
+            self.poll_gpu_stats = True
+        else:
+            self.poll_gpu_stats = False
diff --git a/Engine/Plugins/VirtualProduction/Switchboard/Source/Switchboard/switchboard/devices/ndisplay/ndisplay_monitor_ui.py b/Engine/Plugins/VirtualProduction/Switchboard/Source/Switchboard/switchboard/devices/ndisplay/ndisplay_monitor_ui.py
--- a/Engine/Plugins/VirtualProduction/Switchboard/Source/Switchboard/switchboard/devices/ndisplay/ndisplay_monitor_ui.py
+++ b/Engine/Plugins/VirtualProduction/Switchboard/Source/Switchboard/switchboard/devices/ndisplay/ndisplay_monitor_ui.py
@@ -5,7 +5,7 @@ from typing import Dict, Optional
 from PySide6.QtCore import QSortFilterProxyModel, Qt
 from PySide6.QtWidgets import QComboBox, QCompleter, QHBoxLayout, \
     QHeaderView, QLabel, QPushButton, QSizePolicy, QStyledItemDelegate, \
-    QTableView, QVBoxLayout, QWidget
+    QTableView, QVBoxLayout, QWidget, QCheckBox
 
 from switchboard.devices.ndisplay.ndisplay_monitor import nDisplayMonitor
 
@@ -62,32 +62,48 @@ class nDisplayMonitorUI(QWidget):
             lambda: self.try_issue_console_exec())
 
         # Create remaining buttons
+
+        # Refresh Table button to query all fields in the monitor
         btn_refresh_table = QPushButton('Refresh Table')
         btn_refresh_table.setToolTip(
             'Updates the cached mosaic topologies.')
         btn_refresh_table.clicked.connect(
             self.monitor.on_refresh_table_clicked)
 
-        btn_fix_exe_flags = QPushButton('Fix ExeFlags')
-        btn_fix_exe_flags.setToolTip(
-            'Disables fullscreen optimizations on the executable.')
-        btn_fix_exe_flags.clicked.connect(
-            self.monitor.on_fix_exe_flags_clicked)
-
+        # This button disables full screen optimizations on the executable
+        # This has been linked in the past to being 1 frame out of sync.
+        btn_disable_fso = QPushButton('Disable FSO')
+        btn_disable_fso.setToolTip(
+            'Disables Full Screen Optimizations on the Unreal Engine executable.'
+            '\nNote: The Remote executable must be running.')
+        btn_disable_fso.clicked.connect(
+            self.monitor.on_disable_fso_clicked)
+
+        # Minimize button to minimize all windows on the remote node.
+        # This has been useful in the past when there are windows that
+        # otherwise stay on top of the nDisplay window.
         btn_minimize_windows = QPushButton('Minimize')
         btn_minimize_windows.setToolTip(
             'Minimizes all windows in the nodes.')
         btn_minimize_windows.clicked.connect(
             self.monitor.on_minimize_windows_clicked)
 
+        # GPU Stats checkbox to enable GPU stats monitoring
+        chk_gpu_stats = QCheckBox('GPU Stats')
+        chk_gpu_stats.setToolTip(
+            'GPU Stats add overhead to the remote node that can cause hitches. Only enable when needed.')
+        chk_gpu_stats.stateChanged.connect(
+            self.monitor.on_gpu_stats_toggled)
+
         layout_buttons.addWidget(QLabel('Console:'))
         layout_buttons.addWidget(self.cmb_console_exec)
         layout_buttons.addWidget(btn_console_exec)
         layout_buttons.addStretch(1)
+        layout_buttons.addWidget(chk_gpu_stats)
         layout_buttons.addWidget(btn_refresh_table)
 
-        if self.monitor.use_exe_flags:
-            layout_buttons.addWidget(btn_fix_exe_flags)
+        if self.monitor.show_disable_fso_btn:
+            layout_buttons.addWidget(btn_disable_fso)
 
         layout_buttons.addWidget(btn_minimize_windows)
 
diff --git a/Engine/Plugins/VirtualProduction/Switchboard/Source/Switchboard/switchboard/devices/ndisplay/plugin_ndisplay.py b/Engine/Plugins/VirtualProduction/Switchboard/Source/Switchboard/switchboard/devices/ndisplay/plugin_ndisplay.py
--- a/Engine/Plugins/VirtualProduction/Switchboard/Source/Switchboard/switchboard/devices/ndisplay/plugin_ndisplay.py
+++ b/Engine/Plugins/VirtualProduction/Switchboard/Source/Switchboard/switchboard/devices/ndisplay/plugin_ndisplay.py
@@ -1727,7 +1727,15 @@ class DevicenDisplay(DeviceUnreal):
     # ~ Begin DeviceUnreal Interface
     @QtCore.Slot()
     def _on_listener_connected(self):
+        super()._on_listener_connected()
         if self.__class__.ndisplay_monitor:
             # Poll the sync status to get an immediate update on its state
             self.__class__.ndisplay_monitor.poll_sync_status_for_device(self, SyncStatusRequestFlags.all())
+
+    def on_unreal_started(self):
+        super().on_unreal_started()
+
+        # Always try to disable Full Screen Optimizations. If this is removed, you can add back
+        # the "Disable FSO" button using ndisplay_monitor.show_disable_fso_btn
+        self.disable_fso()
     # ~ End DeviceUnrealInterface
diff --git a/Engine/Plugins/VirtualProduction/Switchboard/Source/Switchboard/switchboard/devices/unreal/plugin_unreal.py b/Engine/Plugins/VirtualProduction/Switchboard/Source/Switchboard/switchboard/devices/unreal/plugin_unreal.py
--- a/Engine/Plugins/VirtualProduction/Switchboard/Source/Switchboard/switchboard/devices/unreal/plugin_unreal.py
+++ b/Engine/Plugins/VirtualProduction/Switchboard/Source/Switchboard/switchboard/devices/unreal/plugin_unreal.py
@@ -969,6 +969,8 @@ class DeviceUnreal(Device):
             super().connecting_listener, QtCore.Qt.QueuedConnection)
         listener_qt_handler.listener_connected.connect(
             super().connect_listener, QtCore.Qt.QueuedConnection)
+        listener_qt_handler.listener_connected.connect(
+            self._on_listener_connected, QtCore.Qt.QueuedConnection)
         listener_qt_handler.listener_connection_failed.connect(
             self._on_listener_connection_failed, QtCore.Qt.QueuedConnection)
 
@@ -1137,6 +1139,10 @@ class DeviceUnreal(Device):
             if DeviceUnreal.rsync_server.is_running():
                 DeviceUnreal.rsync_server.shutdown()
 
+    @QtCore.Slot()
+    def _on_listener_connected(self):
+        pass
+
     @QtCore.Slot()
     def _on_about_to_quit(self):
         if DeviceUnreal.rsync_server.is_running():
@@ -1801,8 +1807,8 @@ class DeviceUnreal(Device):
         else:
             self.status = DeviceStatus.CLOSING
 
-    def fix_exe_flags(self):
-        ''' Tries to force the correct UnrealEditor.exe flags '''
+    def disable_fso(self):
+        ''' Tries to disable Full Screen Optimizations on the remote UnrealEditor.exe executable '''
         unreals = self.program_start_queue.running_programs_named('unreal')
 
         if not len(unreals):
diff --git a/Engine/Source/Programs/SwitchboardListener/SblCore/Internal/SwitchboardListenerVersion.h b/Engine/Source/Programs/SwitchboardListener/SblCore/Internal/SwitchboardListenerVersion.h
--- a/Engine/Source/Programs/SwitchboardListener/SblCore/Internal/SwitchboardListenerVersion.h
+++ b/Engine/Source/Programs/SwitchboardListener/SblCore/Internal/SwitchboardListenerVersion.h
@@ -4,4 +4,4 @@
 
 #define SBLISTENER_VERSION_MAJOR 3
 #define SBLISTENER_VERSION_MINOR 1
-#define SBLISTENER_VERSION_PATCH 0
+#define SBLISTENER_VERSION_PATCH 1
diff --git a/Engine/Source/Programs/SwitchboardListener/SblCore/Private/ScopedNvApi.cpp b/Engine/Source/Programs/SwitchboardListener/SblCore/Private/ScopedNvApi.cpp
new file mode 100644
--- /dev/null
+++ b/Engine/Source/Programs/SwitchboardListener/SblCore/Private/ScopedNvApi.cpp
@@ -0,0 +1,613 @@
+// Copyright Epic Games, Inc. All Rights Reserved.
+
+#if PLATFORM_WINDOWS
+
+#include "ScopedNvApi.h"
+
+#include "Async/RecursiveMutex.h"
+#include "SwitchboardListenerApp.h"
+
+#include "nvapi.h"
+
+
+/** Contains static data relevant to FScopedNvApi */
+class FScopedNvApi::Statics
+{
+public:
+
+	/** Global lock for NvApi usage */
+	static UE::FRecursiveMutex RecursiveMutex;
+
+	/** Keeps track of failed unloads. Useful to recover from potential NvApi unloading errors */
+	static std::atomic<int32> OwedUnloads;
+
+	/** Used to repeatedly show NvAPI initialization errors. Once is enough. */
+	static bool bErrorLoggedOnce;
+
+	/** Cache of the GPU physical handles */
+	static TArray<NvPhysicalGpuHandle> PhysicalGpuHandles;
+
+	/** Keeps count of the number of instances of this class. Used to manage shared resources. */
+	static std::atomic<int32> InstanceCount;
+
+	/** Delegate for when new instances are created */
+	static FOnNvApiInstantiated OnNvApiInstantiated;
+};
+
+
+// Initialize static variables of FScopedNvApi::Statics
+bool                               FScopedNvApi::Statics::bErrorLoggedOnce = false;
+std::atomic<int32>                 FScopedNvApi::Statics::OwedUnloads = 0;
+std::atomic<int32>                 FScopedNvApi::Statics::InstanceCount = 0;
+UE::FRecursiveMutex                FScopedNvApi::Statics::RecursiveMutex;
+TArray<NvPhysicalGpuHandle>        FScopedNvApi::Statics::PhysicalGpuHandles;
+FScopedNvApi::FOnNvApiInstantiated FScopedNvApi::Statics::OnNvApiInstantiated;
+
+
+FScopedNvApi::FOnNvApiInstantiated& FScopedNvApi::GetOnNvApiInstantiated()
+{
+	return Statics::OnNvApiInstantiated;
+}
+
+
+FScopedNvApi::FScopedNvApi()
+{
+	UE::TUniqueLock<UE::FRecursiveMutex> Lock(Statics::RecursiveMutex);
+
+	// Count this new instance
+	Statics::InstanceCount++;
+
+	// Initialize NvAPI.
+
+	if (Statics::OwedUnloads)
+	{
+		Statics::OwedUnloads--;
+		bIsNvApiInitialized = true; // If it isn't true, the calls to the api will simply fail.
+	}
+	else
+	{
+		const NvAPI_Status Result = NvAPI_Initialize();
+
+		if (Result == NVAPI_OK)
+		{
+			bIsNvApiInitialized = true;
+		}
+		else if (!Statics::bErrorLoggedOnce)
+		{
+			Statics::bErrorLoggedOnce = false;
+
+			NvAPI_ShortString ErrorString;
+			NvAPI_GetErrorMessage(Result, ErrorString);
+			UE_LOG(LogSwitchboard, Error, TEXT("NvAPI_Initialize failed. Error: %s"), ANSI_TO_TCHAR(ErrorString));
+		}
+	}
+
+	if (Statics::InstanceCount == 1)
+	{
+		UE_LOG(LogSwitchboard, Verbose, TEXT("FScopedNvApi InstanceCount incremented to 1"));
+	}
+
+	// Let others know that a new instance of this class was instantiated.
+	Statics::OnNvApiInstantiated.Broadcast();
+}
+
+FScopedNvApi::~FScopedNvApi()
+{
+	UE::TUniqueLock<UE::FRecursiveMutex> Lock(Statics::RecursiveMutex);
+
+	// One less instance alive
+	Statics::InstanceCount--;
+
+	check(Statics::InstanceCount >= 0);
+
+	// Clear cache of PhysicalGpuHandles when we go back to zero, forcing a re-cache when
+	// a future instance tries to get a gpu count or read gpu stats.
+	if (!Statics::InstanceCount)
+	{
+		UE_LOG(LogSwitchboard, Verbose, TEXT("FScopedNvApi InstanceCount decremented to 0"));
+		Statics::PhysicalGpuHandles.Reset(); // Keep the memory.
+	}
+
+	// Only unload NvApi if it was initialized
+	if (bIsNvApiInitialized)
+	{
+		const NvAPI_Status Result = NvAPI_Unload();
+
+		if (Result == NVAPI_OK)
+		{
+			bIsNvApiInitialized = false;
+		}
+		else
+		{
+			// For this unexpected event we try to compensate by not loading next time.
+
+			Statics::OwedUnloads++;
+
+			NvAPI_ShortString ErrorString;
+			NvAPI_GetErrorMessage(Result, ErrorString);
+			UE_LOG(LogSwitchboard, Error, TEXT("NvAPI_Unload unexpectedly failed. Error: %s"), ANSI_TO_TCHAR(ErrorString));
+		}
+	}
+
+}
+
+/** Returns true if the NvAPI was initialized successfully */
+bool FScopedNvApi::IsNvApiInitialized() const
+{
+	return bIsNvApiInitialized;
+}
+
+void FScopedNvApi::FillOutSyncTopologies(TArray<FSyncTopo>& SyncTopos) const
+{
+	SWITCHBOARD_TRACE_CPUPROFILER_EVENT_SCOPE(FillOutSyncTopologies);
+
+	UE::TUniqueLock<UE::FRecursiveMutex> Lock(Statics::RecursiveMutex);
+
+	if (!IsNvApiInitialized())
+	{
+		return;
+	}
+
+	// Normally there is a single sync card. BUT an RTX Server could have more, and we need to account for that.
+
+	// Detect sync cards
+
+	NvU32 GSyncCount = 0;
+	NvGSyncDeviceHandle GSyncHandles[NVAPI_MAX_GSYNC_DEVICES];
+	NvAPI_GSync_EnumSyncDevices(GSyncHandles, &GSyncCount); // GSyncCount will be zero if error, so no need to check error.
+
+	for (NvU32 GSyncIdx = 0; GSyncIdx < GSyncCount; GSyncIdx++)
+	{
+		NvU32 GSyncGPUCount = 0;
+		NvU32 GSyncDisplayCount = 0;
+
+		// gather info first with null data pointers, just to get the count and subsequently allocate necessary memory.
+		{
+			const NvAPI_Status Result = NvAPI_GSync_GetTopology(GSyncHandles[GSyncIdx], &GSyncGPUCount, nullptr, &GSyncDisplayCount, nullptr);
+
+			if (Result != NVAPI_OK)
+			{
+				NvAPI_ShortString ErrorString;
+				NvAPI_GetErrorMessage(Result, ErrorString);
+				UE_LOG(LogSwitchboard, Warning, TEXT("NvAPI_GSync_GetTopology failed. Error: %s"), ANSI_TO_TCHAR(ErrorString));
+				continue;
+			}
+		}
+
+		// allocate memory for data
+		TArray<NV_GSYNC_GPU> GSyncGPUs;
+		TArray<NV_GSYNC_DISPLAY> GSyncDisplays;
+		{
+			GSyncGPUs.SetNumUninitialized(GSyncGPUCount, EAllowShrinking::No);
+
+			for (NvU32 GSyncGPUIdx = 0; GSyncGPUIdx < GSyncGPUCount; GSyncGPUIdx++)
+			{
+				GSyncGPUs[GSyncGPUIdx].version = NV_GSYNC_GPU_VER;
+			}
+
+			GSyncDisplays.SetNumUninitialized(GSyncDisplayCount, EAllowShrinking::No);
+
+			for (NvU32 GSyncDisplayIdx = 0; GSyncDisplayIdx < GSyncDisplayCount; GSyncDisplayIdx++)
+			{
+				GSyncDisplays[GSyncDisplayIdx].version = NV_GSYNC_DISPLAY_VER;
+			}
+		}
+
+		// get real info
+		{
+			const NvAPI_Status Result = NvAPI_GSync_GetTopology(GSyncHandles[GSyncIdx], &GSyncGPUCount, GSyncGPUs.GetData(), &GSyncDisplayCount, GSyncDisplays.GetData());
+
+			if (Result != NVAPI_OK)
+			{
+				NvAPI_ShortString ErrorString;
+				NvAPI_GetErrorMessage(Result, ErrorString);
+				UE_LOG(LogSwitchboard, Warning, TEXT("NvAPI_GSync_GetTopology failed. Error: %s"), ANSI_TO_TCHAR(ErrorString));
+				continue;
+			}
+		}
+
+		// Build outbound structure
+
+		FSyncTopo SyncTopo;
+
+		for (NvU32 GpuIdx = 0; GpuIdx < GSyncGPUCount; GpuIdx++)
+		{
+			FSyncGpu SyncGpu;
+
+			SyncGpu.bIsSynced = GSyncGPUs[GpuIdx].isSynced;
+			SyncGpu.Connector = int32(GSyncGPUs[GpuIdx].connector);
+
+			SyncTopo.SyncGpus.Emplace(SyncGpu);
+		}
+
+		for (NvU32 DisplayIdx = 0; DisplayIdx < GSyncDisplayCount; DisplayIdx++)
+		{
+			FSyncDisplay SyncDisplay;
+
+			switch (GSyncDisplays[DisplayIdx].syncState)
+			{
+			case NVAPI_GSYNC_DISPLAY_SYNC_STATE_UNSYNCED:
+				SyncDisplay.SyncState = TEXT("Unsynced");
+				break;
+			case NVAPI_GSYNC_DISPLAY_SYNC_STATE_SLAVE:
+				SyncDisplay.SyncState = TEXT("Follower");
+				break;
+			case NVAPI_GSYNC_DISPLAY_SYNC_STATE_MASTER:
+				SyncDisplay.SyncState = TEXT("Leader");
+				break;
+			default:
+				SyncDisplay.SyncState = TEXT("Unknown");
+				break;
+			}
+
+			// get color information for each display
+			{
+				NV_COLOR_DATA ColorData;
+
+				ColorData.version = NV_COLOR_DATA_VER;
+				ColorData.cmd = NV_COLOR_CMD_GET;
+				ColorData.size = sizeof(NV_COLOR_DATA);
+
+				const NvAPI_Status Result = NvAPI_Disp_ColorControl(GSyncDisplays[DisplayIdx].displayId, &ColorData);
+
+				if (Result == NVAPI_OK)
+				{
+					SyncDisplay.Bpc = ColorData.data.bpc;
+					SyncDisplay.Depth = ColorData.data.depth;
+					SyncDisplay.ColorFormat = ColorData.data.colorFormat;
+				}
+			}
+
+			SyncTopo.SyncDisplays.Emplace(SyncDisplay);
+		}
+
+		// Sync Status Parameters
+		{
+			NV_GSYNC_STATUS_PARAMS GSyncStatusParams;
+			GSyncStatusParams.version = NV_GSYNC_STATUS_PARAMS_VER;
+
+			const NvAPI_Status Result = NvAPI_GSync_GetStatusParameters(GSyncHandles[GSyncIdx], &GSyncStatusParams);
+
+			if (Result != NVAPI_OK)
+			{
+				NvAPI_ShortString ErrorString;
+				NvAPI_GetErrorMessage(Result, ErrorString);
+				UE_LOG(LogSwitchboard, Warning, TEXT("NvAPI_GSync_GetStatusParameters failed. Error: %s"), ANSI_TO_TCHAR(ErrorString));
+				continue;
+			}
+
+			SyncTopo.SyncStatusParams.RefreshRate = GSyncStatusParams.refreshRate;
+			SyncTopo.SyncStatusParams.HouseSyncIncoming = GSyncStatusParams.houseSyncIncoming;
+			SyncTopo.SyncStatusParams.bHouseSync = !!GSyncStatusParams.bHouseSync;
+			SyncTopo.SyncStatusParams.bInternalSecondary = GSyncStatusParams.bInternalSlave;
+		}
+
+		// Sync Control Parameters
+		{
+			NV_GSYNC_CONTROL_PARAMS GSyncControlParams;
+			GSyncControlParams.version = NV_GSYNC_CONTROL_PARAMS_VER;
+
+			const NvAPI_Status Result = NvAPI_GSync_GetControlParameters(GSyncHandles[GSyncIdx], &GSyncControlParams);
+
+			if (Result != NVAPI_OK)
+			{
+				NvAPI_ShortString ErrorString;
+				NvAPI_GetErrorMessage(Result, ErrorString);
+				UE_LOG(LogSwitchboard, Warning, TEXT("NvAPI_GSync_GetControlParameters failed. Error: %s"), ANSI_TO_TCHAR(ErrorString));
+				continue;
+			}
+
+			SyncTopo.SyncControlParams.bInterlaced = !!GSyncControlParams.interlaceMode;
+			SyncTopo.SyncControlParams.bSyncSourceIsOutput = !!GSyncControlParams.syncSourceIsOutput;
+			SyncTopo.SyncControlParams.Interval = GSyncControlParams.interval;
+			SyncTopo.SyncControlParams.Polarity = GSyncControlParams.polarity;
+			SyncTopo.SyncControlParams.Source = GSyncControlParams.source;
+			SyncTopo.SyncControlParams.VMode = GSyncControlParams.vmode;
+
+			SyncTopo.SyncControlParams.SyncSkew.MaxLines = GSyncControlParams.syncSkew.maxLines;
+			SyncTopo.SyncControlParams.SyncSkew.MinPixels = GSyncControlParams.syncSkew.minPixels;
+			SyncTopo.SyncControlParams.SyncSkew.NumLines = GSyncControlParams.syncSkew.numLines;
+			SyncTopo.SyncControlParams.SyncSkew.NumPixels = GSyncControlParams.syncSkew.numPixels;
+
+			SyncTopo.SyncControlParams.StartupDelay.MaxLines = GSyncControlParams.startupDelay.maxLines;
+			SyncTopo.SyncControlParams.StartupDelay.MinPixels = GSyncControlParams.startupDelay.minPixels;
+			SyncTopo.SyncControlParams.StartupDelay.NumLines = GSyncControlParams.startupDelay.numLines;
+			SyncTopo.SyncControlParams.StartupDelay.NumPixels = GSyncControlParams.startupDelay.numPixels;
+		}
+
+		SyncTopos.Emplace(SyncTopo);
+	}
+}
+
+
+void FScopedNvApi::FillOutDriverVersion(FSyncStatus& SyncStatus) const
+{
+	SWITCHBOARD_TRACE_CPUPROFILER_EVENT_SCOPE(FillOutDriverVersion);
+
+	UE::TUniqueLock<UE::FRecursiveMutex> Lock(Statics::RecursiveMutex);
+
+	if (!IsNvApiInitialized())
+	{
+		return;
+	}
+
+	NvU32 DriverVersion;
+	NvAPI_ShortString BuildBranchString;
+
+	const NvAPI_Status Result = NvAPI_SYS_GetDriverAndBranchVersion(&DriverVersion, BuildBranchString);
+
+	if (Result != NVAPI_OK)
+	{
+		NvAPI_ShortString ErrorString;
+		NvAPI_GetErrorMessage(Result, ErrorString);
+		UE_LOG(LogSwitchboard, Warning, TEXT("NvAPI_SYS_GetDriverAndBranchVersion failed. Error: %s"), ANSI_TO_TCHAR(ErrorString));
+		return;
+	}
+
+	SyncStatus.DriverVersion = DriverVersion;
+	SyncStatus.DriverBranch = UTF8_TO_TCHAR(BuildBranchString);
+}
+
+
+uint32 FScopedNvApi::GetGpuCount()
+{
+	UE::TUniqueLock<UE::FRecursiveMutex> Lock(Statics::RecursiveMutex);
+
+	if (!IsNvApiInitialized())
+	{
+		return 0;
+	}
+
+	// Count lazily.
+	if (!Statics::PhysicalGpuHandles.Num())
+	{
+		CacheGpuHandles();
+	}
+
+	return Statics::PhysicalGpuHandles.Num();
+}
+
+
+void FScopedNvApi::FillOutPhysicalGpuStats(FSyncStatus & SyncStatus, bool bGetUtilizations, bool bGetClocks, bool bGetTemperatures)
+{
+	SWITCHBOARD_TRACE_CPUPROFILER_EVENT_SCOPE(FillOutPhysicalGpuStats);
+
+	if (!bGetUtilizations && !bGetClocks && !bGetTemperatures)
+	{
+		return;
+	}
+
+	UE::TUniqueLock<UE::FRecursiveMutex> Lock(Statics::RecursiveMutex);
+
+	if (!IsNvApiInitialized())
+	{
+		return;
+	}
+
+	uint32 PhysicalGpuCount = GetGpuCount();
+
+	SyncStatus.GpuUtilization.SetNumUninitialized(PhysicalGpuCount);
+	SyncStatus.GpuCoreClocksKhz.SetNumUninitialized(PhysicalGpuCount);
+	SyncStatus.GpuTemperature.SetNumUninitialized(PhysicalGpuCount);
+
+	for (NvU32 PhysicalGpuIdx = 0; PhysicalGpuIdx < PhysicalGpuCount; ++PhysicalGpuIdx)
+	{
+		SyncStatus.GpuUtilization[PhysicalGpuIdx] = -1;
+		SyncStatus.GpuCoreClocksKhz[PhysicalGpuIdx] = -1;
+		SyncStatus.GpuTemperature[PhysicalGpuIdx] = MIN_int32;
+
+		const NvPhysicalGpuHandle& PhysicalGpu = Statics::PhysicalGpuHandles[PhysicalGpuIdx];
+
+		NV_GPU_DYNAMIC_PSTATES_INFO_EX PstatesInfo;
+		PstatesInfo.version = NV_GPU_DYNAMIC_PSTATES_INFO_EX_VER;
+		NvAPI_Status NvResult = NvAPI_GPU_GetDynamicPstatesInfoEx(PhysicalGpu, &PstatesInfo);
+		if (NvResult == NVAPI_OK)
+		{
+			// FIXME: NV_GPU_UTILIZATION_DOMAIN_ID enum is missing in our nvapi.h, but documented elsewhere.
+			//const int8 UtilizationDomain = NVAPI_GPU_UTILIZATION_DOMAIN_GPU;
+			const int8 UtilizationDomain = 0;
+			if (PstatesInfo.utilization[UtilizationDomain].bIsPresent)
+			{
+				SyncStatus.GpuUtilization[PhysicalGpuIdx] = PstatesInfo.utilization[UtilizationDomain].percentage;
+			}
+		}
+		else if (NvResult == NVAPI_HANDLE_INVALIDATED)
+		{
+			// If the handles are not valid anymore, we cache new ones and return.
+			// This query will fail but the next one should succeed.
+			UE_LOG(LogSwitchboard, Warning, TEXT("Physical GPU handles were invalidated, recaching them."));
+			CacheGpuHandles();
+			return;
+		}
+		else
+		{
+			UE_LOG(LogSwitchboard, Warning, TEXT("NvAPI_GPU_GetDynamicPstatesInfoEx failed. Error code: %d"), NvResult);
+		}
+
+		NV_GPU_CLOCK_FREQUENCIES ClockFreqs;
+		ClockFreqs.version = NV_GPU_CLOCK_FREQUENCIES_VER;
+		ClockFreqs.ClockType = NV_GPU_CLOCK_FREQUENCIES_CURRENT_FREQ;
+		NvResult = NvAPI_GPU_GetAllClockFrequencies(PhysicalGpu, &ClockFreqs);
+		if (NvResult == NVAPI_OK)
+		{
+			if (ClockFreqs.domain[NVAPI_GPU_PUBLIC_CLOCK_GRAPHICS].bIsPresent)
+			{
+				SyncStatus.GpuCoreClocksKhz[PhysicalGpuIdx] = ClockFreqs.domain[NVAPI_GPU_PUBLIC_CLOCK_GRAPHICS].frequency;
+			}
+		}
+		else if (NvResult == NVAPI_HANDLE_INVALIDATED)
+		{
+			// If the handles are not valid anymore, we cache new ones and return.
+			// This query will fail but the next one should succeed.
+			UE_LOG(LogSwitchboard, Warning, TEXT("Physical GPU handles were invalidated, recaching them."));
+			CacheGpuHandles();
+			return;
+		}
+		else
+		{
+			UE_LOG(LogSwitchboard, Warning, TEXT("NvAPI_GPU_GetAllClockFrequencies failed. Error code: %d"), NvResult);
+		}
+
+		NV_GPU_THERMAL_SETTINGS ThermalSettings;
+		ThermalSettings.version = NV_GPU_THERMAL_SETTINGS_VER;
+		NvResult = NvAPI_GPU_GetThermalSettings(PhysicalGpu, NVAPI_THERMAL_TARGET_ALL, &ThermalSettings);
+		if (NvResult == NVAPI_OK)
+		{
+			// Report max temp across all sensors for this GPU.
+			for (NvU32 SensorIdx = 0; SensorIdx < ThermalSettings.count; ++SensorIdx)
+			{
+				const NvS32 SensorTemp = ThermalSettings.sensor[SensorIdx].currentTemp;
+				if (SensorTemp > SyncStatus.GpuTemperature[PhysicalGpuIdx])
+				{
+					SyncStatus.GpuTemperature[PhysicalGpuIdx] = SensorTemp;
+				}
+			}
+		}
+		else if (NvResult == NVAPI_HANDLE_INVALIDATED)
+		{
+			// If the handles are not valid anymore, we cache new ones and return.
+			// This query will fail but the next one should succeed.
+			UE_LOG(LogSwitchboard, Warning, TEXT("Physical GPU handles were invalidated, recaching them."));
+			CacheGpuHandles();
+			return;
+		}
+		else
+		{
+			UE_LOG(LogSwitchboard, Warning, TEXT("NvAPI_GPU_GetThermalSettings failed. Error code: %d"), NvResult);
+		}
+	}
+}
+
+
+bool FScopedNvApi::CacheGpuHandles()
+{
+	SWITCHBOARD_TRACE_CPUPROFILER_EVENT_SCOPE(CacheGpuHandles);
+
+	UE::TUniqueLock<UE::FRecursiveMutex> Lock(Statics::RecursiveMutex);
+
+	if (!IsNvApiInitialized())
+	{
+		return false;
+	}
+
+	Statics::PhysicalGpuHandles.Reset();
+
+	Statics::PhysicalGpuHandles.SetNumUninitialized(NVAPI_MAX_PHYSICAL_GPUS);
+	NvU32 PhysicalGpuCount;
+	NvAPI_Status NvResult = NvAPI_EnumPhysicalGPUs(Statics::PhysicalGpuHandles.GetData(), &PhysicalGpuCount);
+
+	if (NvResult != NVAPI_OK)
+	{
+		NvAPI_ShortString ErrorString;
+		NvAPI_GetErrorMessage(NvResult, ErrorString);
+		UE_LOG(LogSwitchboard, Warning, TEXT("NvAPI_EnumPhysicalGPUs failed. Error: %s"), ANSI_TO_TCHAR(ErrorString));
+		return false;
+	}
+
+	Statics::PhysicalGpuHandles.SetNum(PhysicalGpuCount);
+
+	// Sort first by bus, then by bus slot, ascending. Consistent with task manager and others.
+	Algo::Sort(Statics::PhysicalGpuHandles, [](const NvPhysicalGpuHandle& Lhs, const NvPhysicalGpuHandle& Rhs) -> bool {
+		NvU32 LhsBusId, RhsBusId, LhsSlotId, RhsSlotId;
+		const NvAPI_Status LhsBusResult = NvAPI_GPU_GetBusId(Lhs, &LhsBusId);
+		const NvAPI_Status RhsBusResult = NvAPI_GPU_GetBusId(Rhs, &RhsBusId);
+		const NvAPI_Status LhsSlotResult = NvAPI_GPU_GetBusSlotId(Lhs, &LhsSlotId);
+		const NvAPI_Status RhsSlotResult = NvAPI_GPU_GetBusSlotId(Rhs, &RhsSlotId);
+
+		if (LhsBusResult != NVAPI_OK || RhsBusResult != NVAPI_OK)
+		{
+			NvAPI_ShortString LhsErrorString;
+			NvAPI_ShortString RhsErrorString;
+			NvAPI_GetErrorMessage(LhsBusResult, LhsErrorString);
+			NvAPI_GetErrorMessage(RhsBusResult, RhsErrorString);
+			UE_LOG(LogSwitchboard, Warning, TEXT("NvAPI_GPU_GetBusId failed. Errors: %s, %s"), ANSI_TO_TCHAR(LhsErrorString), ANSI_TO_TCHAR(RhsErrorString));
+			return false;
+		}
+
+		if (LhsSlotResult != NVAPI_OK || RhsSlotResult != NVAPI_OK)
+		{
+			NvAPI_ShortString LhsErrorString;
+			NvAPI_ShortString RhsErrorString;
+			NvAPI_GetErrorMessage(LhsSlotResult, LhsErrorString);
+			NvAPI_GetErrorMessage(RhsSlotResult, RhsErrorString);
+			UE_LOG(LogSwitchboard, Warning, TEXT("NvAPI_GPU_GetBusSlotId failed. Errors: %s, %s"), ANSI_TO_TCHAR(LhsErrorString), ANSI_TO_TCHAR(RhsErrorString));
+			return false;
+		}
+
+		if (LhsBusId != RhsBusId)
+		{
+			return LhsBusId < RhsBusId;
+		}
+
+		return LhsSlotId < RhsSlotId;
+	});
+
+	return Statics::PhysicalGpuHandles.Num() > 0;
+}
+
+void FScopedNvApi::FillOutMosaicTopologies(TArray<FMosaicTopo>& MosaicTopos)
+{
+	SWITCHBOARD_TRACE_CPUPROFILER_EVENT_SCOPE(FillOutMosaicTopologies);
+
+	UE::TUniqueLock<UE::FRecursiveMutex> Lock(Statics::RecursiveMutex);
+
+	if (!IsNvApiInitialized())
+	{
+		return;
+	}
+
+	NvU32 GridCount = 0;
+	TArray<NV_MOSAIC_GRID_TOPO> GridTopologies;
+
+	// count how many grids
+	{
+		const NvAPI_Status Result = NvAPI_Mosaic_EnumDisplayGrids(nullptr, &GridCount);
+
+		if (Result != NVAPI_OK)
+		{
+			NvAPI_ShortString ErrorString;
+			NvAPI_GetErrorMessage(Result, ErrorString);
+			UE_LOG(LogSwitchboard, Warning, TEXT("NvAPI_Mosaic_EnumDisplayGrids failed. Error: %s"), ANSI_TO_TCHAR(ErrorString));
+			return;
+		}
+	}
+
+	// get the grids
+	{
+		GridTopologies.SetNumUninitialized(GridCount, EAllowShrinking::No);
+
+		for (NvU32 TopoIdx = 0; TopoIdx < GridCount; TopoIdx++)
+		{
+			GridTopologies[TopoIdx].version = NV_MOSAIC_GRID_TOPO_VER;
+		}
+
+		const NvAPI_Status Result = NvAPI_Mosaic_EnumDisplayGrids(GridTopologies.GetData(), &GridCount);
+
+		if (Result != NVAPI_OK)
+		{
+			NvAPI_ShortString ErrorString;
+			NvAPI_GetErrorMessage(Result, ErrorString);
+			UE_LOG(LogSwitchboard, Warning, TEXT("NvAPI_Mosaic_EnumDisplayGrids failed. Error: %s"), ANSI_TO_TCHAR(ErrorString));
+			return;
+		}
+
+		for (NvU32 TopoIdx = 0; TopoIdx < GridCount; TopoIdx++)
+		{
+			FMosaicTopo MosaicTopo;
+			NV_MOSAIC_GRID_TOPO& GridTopo = GridTopologies[TopoIdx];
+
+			MosaicTopo.Columns = GridTopo.columns;
+			MosaicTopo.Rows = GridTopo.rows;
+			MosaicTopo.DisplayCount = GridTopo.displayCount;
+
+			MosaicTopo.DisplaySettings.Bpp = GridTopo.displaySettings.bpp;
+			MosaicTopo.DisplaySettings.Freq = GridTopo.displaySettings.freq;
+			MosaicTopo.DisplaySettings.Height = GridTopo.displaySettings.height;
+			MosaicTopo.DisplaySettings.Width = GridTopo.displaySettings.width;
+
+			MosaicTopos.Emplace(MosaicTopo);
+		}
+	}
+}
+
+#endif // PLATFORM_WINDOWS
diff --git a/Engine/Source/Programs/SwitchboardListener/SblCore/Private/ScopedNvApi.h b/Engine/Source/Programs/SwitchboardListener/SblCore/Private/ScopedNvApi.h
new file mode 100644
--- /dev/null
+++ b/Engine/Source/Programs/SwitchboardListener/SblCore/Private/ScopedNvApi.h
@@ -0,0 +1,61 @@
+// Copyright Epic Games, Inc. All Rights Reserved.
+
+#pragma once
+
+#if PLATFORM_WINDOWS
+
+#include "Delegates/Delegate.h"
+#include "SyncStatus.h"
+
+
+/** Used to allow multi-threaded usage of NvApi and to control when it gets loaded or unloaded */
+class FScopedNvApi
+{
+public:
+
+	FScopedNvApi();
+	~FScopedNvApi();
+
+	/** Returns true if the NvAPI was initialized successfully */
+	bool IsNvApiInitialized() const;
+
+	/** Fills out the Sync topologies */
+	void FillOutSyncTopologies(TArray<FSyncTopo>& SyncTopos) const;
+
+	/** Fills out the sync status */
+	void FillOutDriverVersion(FSyncStatus& SyncStatus) const;
+
+	/** Fills out physical the requested gpu stats */
+	void FillOutPhysicalGpuStats(FSyncStatus& SyncStatus, bool bGetUtilizations, bool bGetClocks, bool bGetTemperatures);
+
+	/** Returns the gpu count */
+	uint32 GetGpuCount();
+
+	/** Fill out the mosaic topologies */
+	void FillOutMosaicTopologies(TArray<FMosaicTopo>& MosaicTopos);
+
+	/** Multicast delegate that gets called when a new instance of FScopedNvApi is instantiated */
+	DECLARE_MULTICAST_DELEGATE(FOnNvApiInstantiated);
+
+	// Delegate for when the index of the single tile being edited changes
+	static FOnNvApiInstantiated& GetOnNvApiInstantiated();
+
+private:
+
+	/** Caches the gpu handles to avoid re-querying them every time gpu stats are requested */
+	bool CacheGpuHandles();
+
+private:
+
+	/** 
+	 * Private static class to encapsulate static members and avoid including nvapi.h in this header,
+	 * which should reduce the temptation of using NvApi functions directly. If that happened, then
+	 * the protections offered by this wrapper would not longer work.
+	 */
+	class Statics;
+
+	/** True if the NvAPI was initialized for use by the current FScopedNvApi instance */
+	bool bIsNvApiInitialized = false;
+};
+
+#endif // PLATFORM_WINDOWS
diff --git a/Engine/Source/Programs/SwitchboardListener/SblCore/Private/SwitchboardListener.cpp b/Engine/Source/Programs/SwitchboardListener/SblCore/Private/SwitchboardListener.cpp
--- a/Engine/Source/Programs/SwitchboardListener/SblCore/Private/SwitchboardListener.cpp
+++ b/Engine/Source/Programs/SwitchboardListener/SblCore/Private/SwitchboardListener.cpp
@@ -20,6 +20,7 @@
 #include "GenericPlatform/GenericPlatformProcess.h"
 #include "HAL/FileManager.h"
 #include "HAL/PlatformFileManager.h"
+#include "HAL/PlatformTime.h"
 #include "Interfaces/IPv4/IPv4Endpoint.h"
 #include "IPAddress.h"
 #include "Logging/StructuredLog.h"
@@ -28,13 +29,11 @@
 #include "Misc/FileHelper.h"
 #include "Misc/Paths.h"
 #include "Misc/ScopeExit.h"
-#include <atomic>
 
+#include <atomic>
 
 #if PLATFORM_WINDOWS
 
-#include "GpuClocker.h"
-
 #pragma warning(push)
 #pragma warning(disable : 4005)	// Disable macro redefinition warning for compatibility with Windows SDK 8+
 
@@ -42,12 +41,11 @@
 #include <Windows.h>
 #include <WinUser.h>
 #include <shellapi.h>
-#include "nvapi.h"
 #include "Windows/HideWindowsPlatformTypes.h"
 
-#pragma warning(pop)
+#include "ScopedNvApi.h"
 
-static void FillOutMosaicTopologies(TArray<FMosaicTopo>& MosaicTopos);
+#pragma warning(pop)
 
 #endif
 
@@ -160,6 +158,85 @@ namespace
 		}
 		return TEXT("Unknown MsQuic error.");
 	}
+
+#if PLATFORM_WINDOWS
+
+	/** 
+	 * Manages an instance of FScopedNvApi to keep the library initialized until it doesn't get used for a while.
+	 * The objective is to work around an issue where NvApi can sometimes cause hitches after it has been initialized for 
+	 * around 12 hours. If the client is disconnected or is not requesting data that requires NvApi, then this
+	 * object will not keep an instance of FScopedNvApi so that the libary can be uninitialized. But if it is being
+	 * actively used, then it will keep this object alive and hence the library initialized for normal usage.
+	 */
+	class FWorkaroundForHitchingAfterHours
+	{
+	private:
+
+		FWorkaroundForHitchingAfterHours()
+		{
+			FScopedNvApi::GetOnNvApiInstantiated().AddRaw(this, &FWorkaroundForHitchingAfterHours::OnNvApiInstantiated);
+		}
+
+		~FWorkaroundForHitchingAfterHours()
+		{
+			FScopedNvApi::GetOnNvApiInstantiated().RemoveAll(this);
+		}
+
+		/** Called when a new insance of FScopedNvApi is created. This will signal that it is being used elsewhere. */
+		void OnNvApiInstantiated()
+		{
+			LastTimeInstantiatedSeconds = FPlatformTime::Seconds();
+			bNvApiInstanceNeeded = true;
+		}
+
+	public:
+
+		/** Singleton getter. We only want one instance of this workaround class. */
+		static FWorkaroundForHitchingAfterHours& Get()
+		{
+			static FWorkaroundForHitchingAfterHours Instance;
+			return Instance;
+		}
+
+		/** Call this periodically in the main thread to update the internal instance of FScopedNvApi */
+		void Tick()
+		{
+			if (bNvApiInstanceNeeded)
+			{
+				if (!ScopedNvApi.IsValid())
+				{
+					ScopedNvApi = MakeShared<FScopedNvApi>();
+				}
+
+				bNvApiInstanceNeeded = false;
+			}
+			else if (ScopedNvApi.IsValid())
+			{
+				const double NowSeconds = FPlatformTime::Seconds();
+				const double ElapsedSeconds = NowSeconds - LastTimeInstantiatedSeconds;
+
+				if (ElapsedSeconds > NvApiUsageTimeoutMinutes * 60)
+				{
+					ScopedNvApi.Reset();
+				}
+			}
+		}
+
+	private:
+
+		/** No new instancing of FScopedNvApi for this time will trigger a reset of the local instance kept by this object */
+		const double NvApiUsageTimeoutMinutes = 10;
+
+		/** Signals a request of creating a new instance of FScopedNvApi */
+		std::atomic<bool> bNvApiInstanceNeeded = false;
+
+		/** A record of the last time that an instance of FScopedNvApi was created. Use to detect usage timeouts. */
+		std::atomic<double> LastTimeInstantiatedSeconds = 0;
+
+		/** The instance of ScopedNvApi that will be kept alive until a timout of lack of usage of the library */
+		TSharedPtr<FScopedNvApi> ScopedNvApi;
+	};
+#endif // PLATFORM_WINDOWS
 }
 
 struct FRunningProcess
@@ -290,98 +367,6 @@ FString FSwitchboardCommandLineOptions::ToString(bool bIncludeRedeploy /* = fals
 }
 
 
-#if PLATFORM_WINDOWS
-/** Used to limit NvApi loading to the current scope (recursive allowed) and thread. */
-class FScopedNvApi
-{
-private:
-
-	/** Global recursive lock for NvApi usage */
-	static UE::FRecursiveMutex NvApiLock;
-
-	/** Keeps track of failed unloads. Useful to recover from potential NvApi unloading errors */
-	static std::atomic<int32> OwedUnloads;
-
-	/** Used to repeatedly show NvAPI initialization errors. Once is enough. */
-	static bool bErrorLoggedOnce;
-
-	/** True if the NvAPI was initialized */
-	bool bIsNvApiInitialized;
-
-public:
-	FScopedNvApi()
-		: bIsNvApiInitialized(false)
-	{
-		NvApiLock.Lock(); // Unlocked at the end of the destructor
-
-		// Initialize NvAPI.
-
-		if (OwedUnloads)
-		{
-			OwedUnloads--;
-			bIsNvApiInitialized = true; // If it isn't true, the calls to the api will simply fail.
-		}
-		else
-		{
-			const NvAPI_Status Result = NvAPI_Initialize();
-
-			if (Result == NVAPI_OK)
-			{
-				bIsNvApiInitialized = true;
-			}
-			else if (!bErrorLoggedOnce)
-			{
-				bErrorLoggedOnce = false;
-
-				NvAPI_ShortString ErrorString;
-				NvAPI_GetErrorMessage(Result, ErrorString);
-				UE_LOG(LogSwitchboard, Error, TEXT("NvAPI_Initialize failed. Error: %s"), ANSI_TO_TCHAR(ErrorString));
-			}
-		}
-	}
-
-	~FScopedNvApi()
-	{
-		// Only unload NvApi if it was initialized
-		if (bIsNvApiInitialized)
-		{
-			const NvAPI_Status Result = NvAPI_Unload();
-
-			if (Result == NVAPI_OK)
-			{
-				bIsNvApiInitialized = false;
-			}
-			else
-			{
-				// This would be an unexpected event because the lock prevents concurrent usage of the library.
-				// We try to compensate by not loading next time.
-
-				OwedUnloads++;
-
-				NvAPI_ShortString ErrorString;
-				NvAPI_GetErrorMessage(Result, ErrorString);
-				UE_LOG(LogSwitchboard, Error, TEXT("NvAPI_Unload unexpectedly failed. Error: %s"), ANSI_TO_TCHAR(ErrorString));
-			}
-		}
-
-		NvApiLock.Unlock();
-	}
-
-	/** Returns true if the NvAPI was initialized successfully */
-	bool IsNvApiInitialized()
-	{
-		return bIsNvApiInitialized;
-	}
-};
-
-// Initialize static member variables
-bool FScopedNvApi::bErrorLoggedOnce = false;
-std::atomic<int32> FScopedNvApi::OwedUnloads = 0;
-UE::FRecursiveMutex FScopedNvApi::NvApiLock;
-
-#endif // PLATFORM_WINDOWS
-
-
 FSwitchboardListener::FSwitchboardListener(const FSwitchboardCommandLineOptions& InOptions)
 	: Options(InOptions)
 	, CpuMonitor(MakeShared<FCpuUtilizationMonitor>())
@@ -391,7 +376,8 @@ FSwitchboardListener::FSwitchboardListener(const FSwitchboardCommandLineOptions&
 {
 #if PLATFORM_WINDOWS
 	// Cache Mosaic Topologies
-	FillOutMosaicTopologies(*CachedMosaicTopos);
+	FScopedNvApi NvApi;
+	NvApi.FillOutMosaicTopologies(*CachedMosaicTopos);
 #endif // PLATFORM_WINDOWS
 
 	const int32 NumLogicalProcessors = FPlatformMisc::NumberOfCoresIncludingHyperthreads();
@@ -635,6 +621,10 @@ void FSwitchboardListener::Tick()
 	SendMessageFutures();
 	SBLHelper->Tick();
 
+#if PLATFORM_WINDOWS
+	FWorkaroundForHitchingAfterHours::Get().Tick();
+#endif // PLATFORM_WINDOWS
+
 	OnTickDelegate.Broadcast();
 }
 
@@ -1900,220 +1890,6 @@ stream_receive_outer_break:
 }
 
 
-#if PLATFORM_WINDOWS
-static void FillOutSyncTopologies(TArray<FSyncTopo>& SyncTopos)
-{
-	SWITCHBOARD_TRACE_CPUPROFILER_EVENT_SCOPE(FillOutSyncTopologies);
-
-	FScopedNvApi ScopedNvApi;
-
-	if (!ScopedNvApi.IsNvApiInitialized())
-	{
-		return;
-	}
-
-	// Normally there is a single sync card. BUT an RTX Server could have more, and we need to account for that.
-
-	// Detect sync cards
-
-	NvU32 GSyncCount = 0;
-	NvGSyncDeviceHandle GSyncHandles[NVAPI_MAX_GSYNC_DEVICES];
-	NvAPI_GSync_EnumSyncDevices(GSyncHandles, &GSyncCount); // GSyncCount will be zero if error, so no need to check error.
-
-	for (NvU32 GSyncIdx = 0; GSyncIdx < GSyncCount; GSyncIdx++)
-	{
-		NvU32 GSyncGPUCount = 0;
-		NvU32 GSyncDisplayCount = 0;
-
-		// gather info first with null data pointers, just to get the count and subsequently allocate necessary memory.
-		{
-			const NvAPI_Status Result = NvAPI_GSync_GetTopology(GSyncHandles[GSyncIdx], &GSyncGPUCount, nullptr, &GSyncDisplayCount, nullptr);
-
-			if (Result != NVAPI_OK)
-			{
-				NvAPI_ShortString ErrorString;
-				NvAPI_GetErrorMessage(Result, ErrorString);
-				UE_LOG(LogSwitchboard, Warning, TEXT("NvAPI_GSync_GetTopology failed. Error: %s"), ANSI_TO_TCHAR(ErrorString));
-				continue;
-			}
-		}
-
-		// allocate memory for data
-		TArray<NV_GSYNC_GPU> GSyncGPUs;
-		TArray<NV_GSYNC_DISPLAY> GSyncDisplays;
-		{
-			GSyncGPUs.SetNumUninitialized(GSyncGPUCount, EAllowShrinking::No);
-
-			for (NvU32 GSyncGPUIdx = 0; GSyncGPUIdx < GSyncGPUCount; GSyncGPUIdx++)
-			{
-				GSyncGPUs[GSyncGPUIdx].version = NV_GSYNC_GPU_VER;
-			}
-
-			GSyncDisplays.SetNumUninitialized(GSyncDisplayCount, EAllowShrinking::No);
-
-			for (NvU32 GSyncDisplayIdx = 0; GSyncDisplayIdx < GSyncDisplayCount; GSyncDisplayIdx++)
-			{
-				GSyncDisplays[GSyncDisplayIdx].version = NV_GSYNC_DISPLAY_VER;
-			}
-		}
-
-		// get real info
-		{
-			const NvAPI_Status Result = NvAPI_GSync_GetTopology(GSyncHandles[GSyncIdx], &GSyncGPUCount, GSyncGPUs.GetData(), &GSyncDisplayCount, GSyncDisplays.GetData());
-
-			if (Result != NVAPI_OK)
-			{
-				NvAPI_ShortString ErrorString;
-				NvAPI_GetErrorMessage(Result, ErrorString);
-				UE_LOG(LogSwitchboard, Warning, TEXT("NvAPI_GSync_GetTopology failed. Error: %s"), ANSI_TO_TCHAR(ErrorString));
-				continue;
-			}
-		}
-
-		// Build outbound structure
-
-		FSyncTopo SyncTopo;
-
-		for (NvU32 GpuIdx = 0; GpuIdx < GSyncGPUCount; GpuIdx++)
-		{
-			FSyncGpu SyncGpu;
-
-			SyncGpu.bIsSynced = GSyncGPUs[GpuIdx].isSynced;
-			SyncGpu.Connector = int32(GSyncGPUs[GpuIdx].connector);
-
-			SyncTopo.SyncGpus.Emplace(SyncGpu);
-		}
-
-		for (NvU32 DisplayIdx = 0; DisplayIdx < GSyncDisplayCount; DisplayIdx++)
-		{
-			FSyncDisplay SyncDisplay;
-
-			switch (GSyncDisplays[DisplayIdx].syncState)
-			{
-			case NVAPI_GSYNC_DISPLAY_SYNC_STATE_UNSYNCED:
-				SyncDisplay.SyncState = TEXT("Unsynced");
-				break;
-			case NVAPI_GSYNC_DISPLAY_SYNC_STATE_SLAVE:
-				SyncDisplay.SyncState = TEXT("Follower");
-				break;
-			case NVAPI_GSYNC_DISPLAY_SYNC_STATE_MASTER:
-				SyncDisplay.SyncState = TEXT("Leader");
-				break;
-			default:
-				SyncDisplay.SyncState = TEXT("Unknown");
-				break;
-			}
-
-			// get color information for each display
-			{
-				NV_COLOR_DATA ColorData;
-
-				ColorData.version = NV_COLOR_DATA_VER;
-				ColorData.cmd = NV_COLOR_CMD_GET;
-				ColorData.size = sizeof(NV_COLOR_DATA);
-
-				const NvAPI_Status Result = NvAPI_Disp_ColorControl(GSyncDisplays[DisplayIdx].displayId, &ColorData);
-
-				if (Result == NVAPI_OK)
-				{
-					SyncDisplay.Bpc = ColorData.data.bpc;
-					SyncDisplay.Depth = ColorData.data.depth;
-					SyncDisplay.ColorFormat = ColorData.data.colorFormat;
-				}
-			}
-
-			SyncTopo.SyncDisplays.Emplace(SyncDisplay);
-		}
-
-		// Sync Status Parameters
-		{
-			NV_GSYNC_STATUS_PARAMS GSyncStatusParams;
-			GSyncStatusParams.version = NV_GSYNC_STATUS_PARAMS_VER;
-
-			const NvAPI_Status Result = NvAPI_GSync_GetStatusParameters(GSyncHandles[GSyncIdx], &GSyncStatusParams);
-
-			if (Result != NVAPI_OK)
-			{
-				NvAPI_ShortString ErrorString;
-				NvAPI_GetErrorMessage(Result, ErrorString);
-				UE_LOG(LogSwitchboard, Warning, TEXT("NvAPI_GSync_GetStatusParameters failed. Error: %s"), ANSI_TO_TCHAR(ErrorString));
-				continue;
-			}
-
-			SyncTopo.SyncStatusParams.RefreshRate = GSyncStatusParams.refreshRate;
-			SyncTopo.SyncStatusParams.HouseSyncIncoming = GSyncStatusParams.houseSyncIncoming;
-			SyncTopo.SyncStatusParams.bHouseSync = !!GSyncStatusParams.bHouseSync;
-			SyncTopo.SyncStatusParams.bInternalSecondary = GSyncStatusParams.bInternalSlave;
-		}
-
-		// Sync Control Parameters
-		{
-			NV_GSYNC_CONTROL_PARAMS GSyncControlParams;
-			GSyncControlParams.version = NV_GSYNC_CONTROL_PARAMS_VER;
-
-			const NvAPI_Status Result = NvAPI_GSync_GetControlParameters(GSyncHandles[GSyncIdx], &GSyncControlParams);
-
-			if (Result != NVAPI_OK)
-			{
-				NvAPI_ShortString ErrorString;
-				NvAPI_GetErrorMessage(Result, ErrorString);
-				UE_LOG(LogSwitchboard, Warning, TEXT("NvAPI_GSync_GetControlParameters failed. Error: %s"), ANSI_TO_TCHAR(ErrorString));
-				continue;
-			}
-
-			SyncTopo.SyncControlParams.bInterlaced = !!GSyncControlParams.interlaceMode;
-			SyncTopo.SyncControlParams.bSyncSourceIsOutput = !!GSyncControlParams.syncSourceIsOutput;
-			SyncTopo.SyncControlParams.Interval = GSyncControlParams.interval;
-			SyncTopo.SyncControlParams.Polarity = GSyncControlParams.polarity;
-			SyncTopo.SyncControlParams.Source = GSyncControlParams.source;
-			SyncTopo.SyncControlParams.VMode = GSyncControlParams.vmode;
-
-			SyncTopo.SyncControlParams.SyncSkew.MaxLines = GSyncControlParams.syncSkew.maxLines;
-			SyncTopo.SyncControlParams.SyncSkew.MinPixels = GSyncControlParams.syncSkew.minPixels;
-			SyncTopo.SyncControlParams.SyncSkew.NumLines = GSyncControlParams.syncSkew.numLines;
-			SyncTopo.SyncControlParams.SyncSkew.NumPixels = GSyncControlParams.syncSkew.numPixels;
-
-			SyncTopo.SyncControlParams.StartupDelay.MaxLines = GSyncControlParams.startupDelay.maxLines;
-			SyncTopo.SyncControlParams.StartupDelay.MinPixels = GSyncControlParams.startupDelay.minPixels;
-			SyncTopo.SyncControlParams.StartupDelay.NumLines = GSyncControlParams.startupDelay.numLines;
-			SyncTopo.SyncControlParams.StartupDelay.NumPixels = GSyncControlParams.startupDelay.numPixels;
-		}
-
-		SyncTopos.Emplace(SyncTopo);
-	}
-}
-#endif // PLATFORM_WINDOWS
-
-#if PLATFORM_WINDOWS
-static void FillOutDriverVersion(FSyncStatus& SyncStatus)
-{
-	SWITCHBOARD_TRACE_CPUPROFILER_EVENT_SCOPE(FillOutDriverVersion);
-
-	FScopedNvApi ScopedNvApi;
-
-	if (!ScopedNvApi.IsNvApiInitialized())
-	{
-		return;
-	}
-
-	NvU32 DriverVersion;
-	NvAPI_ShortString BuildBranchString;
-
-	const NvAPI_Status Result = NvAPI_SYS_GetDriverAndBranchVersion(&DriverVersion, BuildBranchString);
-
-	if (Result != NVAPI_OK)
-	{
-		NvAPI_ShortString ErrorString;
-		NvAPI_GetErrorMessage(Result, ErrorString);
-		UE_LOG(LogSwitchboard, Warning, TEXT("NvAPI_SYS_GetDriverAndBranchVersion failed. Error: %s"), ANSI_TO_TCHAR(ErrorString));
-		return;
-	}
-
-	SyncStatus.DriverVersion = DriverVersion;
-	SyncStatus.DriverBranch = UTF8_TO_TCHAR(BuildBranchString);
-}
-#endif // PLATFORM_WINDOWS
-
 #if PLATFORM_WINDOWS
 static void FillOutTaskbarAutoHide(FSyncStatus& SyncStatus)
 {
@@ -2138,73 +1914,6 @@ static void FillOutTaskbarAutoHide(FSyncStatus& SyncStatus)
 #endif // PLATFORM_WINDOWS
 
 
-#if PLATFORM_WINDOWS
-static void FillOutMosaicTopologies(TArray<FMosaicTopo>& MosaicTopos)
-{
-	SWITCHBOARD_TRACE_CPUPROFILER_EVENT_SCOPE(FillOutMosaicTopologies);
-
-	FScopedNvApi ScopedNvApi;
-
-	if (!ScopedNvApi.IsNvApiInitialized())
-	{
-		return;
-	}
-
-	NvU32 GridCount = 0;
-	TArray<NV_MOSAIC_GRID_TOPO> GridTopologies;
-
-	// count how many grids
-	{
-		const NvAPI_Status Result = NvAPI_Mosaic_EnumDisplayGrids(nullptr, &GridCount);
-
-		if (Result != NVAPI_OK)
-		{
-			NvAPI_ShortString ErrorString;
-			NvAPI_GetErrorMessage(Result, ErrorString);
-			UE_LOG(LogSwitchboard, Warning, TEXT("NvAPI_Mosaic_EnumDisplayGrids failed. Error: %s"), ANSI_TO_TCHAR(ErrorString));
-			return;
-		}
-	}
-
-	// get the grids
-	{
-		GridTopologies.SetNumUninitialized(GridCount, EAllowShrinking::No);
-
-		for (NvU32 TopoIdx = 0; TopoIdx < GridCount; TopoIdx++)
-		{
-			GridTopologies[TopoIdx].version = NV_MOSAIC_GRID_TOPO_VER;
-		}
-
-		const NvAPI_Status Result = NvAPI_Mosaic_EnumDisplayGrids(GridTopologies.GetData(), &GridCount);
-
-		if (Result != NVAPI_OK)
-		{
-			NvAPI_ShortString ErrorString;
-			NvAPI_GetErrorMessage(Result, ErrorString);
-			UE_LOG(LogSwitchboard, Warning, TEXT("NvAPI_Mosaic_EnumDisplayGrids failed. Error: %s"), ANSI_TO_TCHAR(ErrorString));
-			return;
-		}
-
-		for (NvU32 TopoIdx = 0; TopoIdx < GridCount; TopoIdx++)
-		{
-			FMosaicTopo MosaicTopo;
-			NV_MOSAIC_GRID_TOPO& GridTopo = GridTopologies[TopoIdx];
-
-			MosaicTopo.Columns = GridTopo.columns;
-			MosaicTopo.Rows = GridTopo.rows;
-			MosaicTopo.DisplayCount = GridTopo.displayCount;
-
-			MosaicTopo.DisplaySettings.Bpp = GridTopo.displaySettings.bpp;
-			MosaicTopo.DisplaySettings.Freq = GridTopo.displaySettings.freq;
-			MosaicTopo.DisplaySettings.Height = GridTopo.displaySettings.height;
-			MosaicTopo.DisplaySettings.Width = GridTopo.displaySettings.width;
-
-			MosaicTopos.Emplace(MosaicTopo);
-		}
-	}
-}
-#endif // PLATFORM_WINDOWS
-
 #if PLATFORM_WINDOWS
 FRunningProcess* FSwitchboardListener::FindOrStartFlipModeMonitorForUUID(const FGuid& UUID)
 {
@@ -2406,69 +2115,6 @@ static void FillOutDisableFullscreenOptimizationForProcess(FSyncStatus& SyncStat
 }
 #endif // PLATFORM_WINDOWS
 
-#if PLATFORM_WINDOWS
-static void FillOutPhysicalGpuStats(FSyncStatus& SyncStatus, bool bGetUtilizations, bool bGetClocks, bool bGetTemperatures)
-{
-	SWITCHBOARD_TRACE_CPUPROFILER_EVENT_SCOPE(FillOutPhysicalGpuStats);
-
-	if (!bGetUtilizations && !bGetClocks && !bGetTemperatures)
-	{
-		return;
-	}
-
-	FGpuClocker GpuClocker;
-
-	uint32 GpuCount;
-
-	if (!GpuClocker.GetGpuCount(GpuCount))
-	{
-		return;
-	}
-
-	SyncStatus.GpuUtilization.SetNumUninitialized(GpuCount);
-	SyncStatus.GpuCoreClocksKhz.SetNumUninitialized(GpuCount);
-	SyncStatus.GpuTemperature.SetNumUninitialized(GpuCount);
-
-	for (uint32 GpuIdx = 0; GpuIdx < GpuCount; ++GpuIdx)
-	{
-		SyncStatus.GpuUtilization[GpuIdx] = -1;
-		SyncStatus.GpuCoreClocksKhz[GpuIdx] = -1;
-		SyncStatus.GpuTemperature[GpuIdx] = MIN_int32;
-
-		if (bGetUtilizations)
-		{
-			uint32 UsageKernel;
-			uint32 UsageMemory;
-
-			if (GpuClocker.GetGpuUsage(GpuIdx, UsageKernel, UsageMemory))
-			{
-				SyncStatus.GpuUtilization[GpuIdx] = static_cast<int8>(UsageKernel);
-			}
-		}
-
-		if (bGetClocks)
-		{
-			uint32 GraphicsMHz;
-			uint32 MemoryMHz;
-
-			if (GpuClocker.GetGpuMHz(GpuIdx, GraphicsMHz, MemoryMHz))
-			{
-				SyncStatus.GpuCoreClocksKhz[GpuIdx] = GraphicsMHz * 1000;
-			}
-		}
-
-		if (bGetTemperatures)
-		{
-			uint32 Celsius;
-
-			if (GpuClocker.GetGpuCelsius(GpuIdx, Celsius))
-			{
-				SyncStatus.GpuTemperature[GpuIdx] = static_cast<int32>(Celsius);
-			}
-		}
-	}
-}
-#endif
 
 bool FSwitchboardListener::EquivalentTaskFutureExists(uint32 TaskEquivalenceHash) const
 {
@@ -2547,12 +2193,12 @@ bool FSwitchboardListener::Task_GetSyncStatus(const FSwitchboardGetSyncStatusTas
 
 				if (EnumHasAnyFlags(RequestFlags, ESyncStatusRequestFlags::DriverInfo))
 				{
-					FillOutDriverVersion(SyncStatus.Get());
+					ScopedNvApi.FillOutDriverVersion(SyncStatus.Get());
 				}
 
 				if (EnumHasAnyFlags(RequestFlags, ESyncStatusRequestFlags::SyncTopos))
 				{
-					FillOutSyncTopologies(SyncStatus->SyncTopos);
+					ScopedNvApi.FillOutSyncTopologies(SyncStatus->SyncTopos);
 				}
 			}
 
@@ -2589,7 +2235,11 @@ bool FSwitchboardListener::Task_GetSyncStatus(const FSwitchboardGetSyncStatusTas
 				const bool bGetClocks = EnumHasAnyFlags(RequestFlags, ESyncStatusRequestFlags::GpuCoreClockKhz);
 				const bool bGetTemperatures = EnumHasAnyFlags(RequestFlags, ESyncStatusRequestFlags::GpuTemperature);
 
-				FillOutPhysicalGpuStats(SyncStatus.Get(), bGetUtilizations, bGetClocks, bGetTemperatures);
+				if (bGetUtilizations || bGetClocks || bGetTemperatures)
+				{
+					FScopedNvApi ScopedNvApi;
+					ScopedNvApi.FillOutPhysicalGpuStats(SyncStatus.Get(), bGetUtilizations, bGetClocks, bGetTemperatures);
+				}
 			}
 
 			return CreateSyncStatusMessage(SyncStatus.Get(), RequestFlags);
@@ -2647,7 +2297,9 @@ bool FSwitchboardListener::Task_RefreshMosaics(const FSwitchboardRefreshMosaicsT
 
 			FWriteScopeLock Lock(*CachedMosaicToposLock);
 			CachedMosaicTopos->Reset();
-			FillOutMosaicTopologies(*CachedMosaicTopos);
+
+			FScopedNvApi ScopedNvApi;
+			ScopedNvApi.FillOutMosaicTopologies(*CachedMosaicTopos);
 			
 			return FString();
 		}
