SBL: It can now lock the gpus itself. SBLH is not required anymore.

For now we are keeping SBLH around mostly so that the user doesn't need to run SBL as admin.

Changed a log from ScopedNvApi to VeryVerbose as to not clutter the log window.

Also fixed an issue where the open button icon of unreal devices in SB would not always update because it would update before the program queue was updated after a program ended notification came from the listener, and that program queue was checked when updating the classes.

diff --git a/Engine/Plugins/VirtualProduction/Switchboard/Source/Switchboard/switchboard/devices/unreal/plugin_unreal.py b/Engine/Plugins/VirtualProduction/Switchboard/Source/Switchboard/switchboard/devices/unreal/plugin_unreal.py
--- a/Engine/Plugins/VirtualProduction/Switchboard/Source/Switchboard/switchboard/devices/unreal/plugin_unreal.py
+++ b/Engine/Plugins/VirtualProduction/Switchboard/Source/Switchboard/switchboard/devices/unreal/plugin_unreal.py
@@ -2346,8 +2346,6 @@ class DeviceUnreal(Device):
             f"{self.name}: Program with id {puuid} exited with "
             f"returncode {returncode}")
 
-        self._update_widget_classes()
-
         try:
             program_name = self.program_start_queue.on_program_ended(
                 puuid=puuid, unreal_client=self.unreal_client)
@@ -2356,6 +2354,11 @@ class DeviceUnreal(Device):
                 f"{self.name}: on_program_ended with unknown id {puuid}")
             return
 
+        # Update the widget classes (e.g. open button icon) once the
+        # program_start_queue has been updated, since the status of the
+        # queue may affect the widget classes.
+        self._update_widget_classes()
+
         # Check if there are remaining programs named the same but with
         # different ids, which is not normal.
         remaining_homonyms = self.program_start_queue.running_puuids_named(
diff --git a/Engine/Source/Programs/SwitchboardListener/SblCore/Internal/SwitchboardListener.h b/Engine/Source/Programs/SwitchboardListener/SblCore/Internal/SwitchboardListener.h
--- a/Engine/Source/Programs/SwitchboardListener/SblCore/Internal/SwitchboardListener.h
+++ b/Engine/Source/Programs/SwitchboardListener/SblCore/Internal/SwitchboardListener.h
@@ -39,7 +39,7 @@ struct FSwitchboardMinimizeWindowsTask;
 struct FSwitchboardSetInactiveTimeoutTask;
 struct FSwitchboardFreeListenerBinaryTask;
 
-class FSBLHelperClient;
+class FGpuClockerManager;
 
 
 struct FSwitchboardCommandLineOptions
@@ -264,8 +264,8 @@ private:
 	TArray<FSwitchboardMessageFuture> MessagesFutures;
 	TSharedPtr<FCpuUtilizationMonitor> CpuMonitor;
 
-	/** Client interface to the Switchboard Listener Helper external process */
-	TSharedPtr<FSBLHelperClient> SBLHelper;
+	/** Gpu clock manager. Used to lock/unlock gpu clocks. */
+	TSharedPtr<FGpuClockerManager> GpuClockerManager;
 
 	bool bProcessorSMT;
 
diff --git a/Engine/Source/Programs/SwitchboardListener/SblCore/Private/GpuClockerManager.cpp b/Engine/Source/Programs/SwitchboardListener/SblCore/Private/GpuClockerManager.cpp
new file mode 100644
--- /dev/null
+++ b/Engine/Source/Programs/SwitchboardListener/SblCore/Private/GpuClockerManager.cpp
@@ -0,0 +1,150 @@
+// Copyright Epic Games, Inc. All Rights Reserved.
+
+#include "GpuClockerManager.h"
+
+#include "GpuClocker.h"
+#include "Misc/CommandLine.h"
+#include "SBLHelperClient.h"
+#include "SwitchboardListenerApp.h" // To include LogSwitchboard
+
+
+FGpuClockerManager::FGpuClockerManager()
+	: SBLHelper(MakeShared<FSBLHelperClient>())
+{ 
+}
+
+FGpuClockerManager::~FGpuClockerManager()
+{
+	// Unlock clocks when exiting. We do not expect SBL to be closed while it is managing processes but if it happens
+	// it should try to leave them in a normal state. In case of abnormal termination, this code won't run
+	// and the gpus will be left as they were.
+	if (bLockManagedLocally)
+	{
+		FGpuClocker GpuClocker;
+		GpuClocker.UnlockGpuClocks();
+	}
+}
+
+
+bool FGpuClockerManager::LockGpuClocksForPid(uint32 Pid)
+{
+	// Add the Pid to the set. We do this anyway whether we fully succeed or not in the code below
+	// to account for the case of a GPU succeeding and another one not succeeding. In that case,
+	// we want to proceeed and not get into a situation of a GPU actually getting locked but never unlocking it
+	// because we thought it failed.
+	LockingPids.Add(Pid);
+	
+	// Attempt to lock via the SBL Helper if we're not locally managing the lock.
+	if (!bLockManagedLocally && LockGpuClocksUsingSBLHelper(Pid))
+	{
+		return true;
+	}
+
+	// If SBLH is not available then we attempt to do it locally.
+	UE_LOG(LogSwitchboard, Display, TEXT("Locking Gpu Clocks"));
+	bLockManagedLocally = true;
+
+	FGpuClocker GpuClocker;
+	const bool bLocked = GpuClocker.LockGpuClocks();
+
+	if (!bLocked)
+	{
+		UE_LOG(LogSwitchboard, Warning, TEXT(
+			"Unable to lock Gpu clocks as requested in Switchboard settings for process with id %d. Please run SwitchboardListenerHelper "
+			"or SwitchboardListener with elevated privileges because NVML requires this in order to lock the Gpu clocks."), Pid);
+	}
+
+	return bLocked;
+}
+
+
+void FGpuClockerManager::PidEnded(uint32 Pid)
+{
+	// Remove the ended pid from our list that keeps the gpus locked.
+	// If nothing was removed, then we just return as there is nothing to do.
+	if (!LockingPids.Remove(Pid))
+	{
+		return;
+	}
+
+	// If we just emptied the pids and are locally managing the lock, then we locally unlock the gpu clocks.
+	if (bLockManagedLocally && LockingPids.IsEmpty())
+	{
+		UE_LOG(LogSwitchboard, Display, TEXT("Unlocking Gpu Clocks"));
+
+		// Since we're unlocking, we cease to assert that we are managing the lock locally.
+		// This makes it so that it tries SBL Helper first next time.
+		bLockManagedLocally = false;
+
+		FGpuClocker GpuClocker;
+		GpuClocker.UnlockGpuClocks();
+	}
+}
+
+
+void FGpuClockerManager::Tick()
+{
+	if (SBLHelper.IsValid())
+	{
+		SBLHelper->Tick();
+	}
+}
+
+
+bool FGpuClockerManager::LockGpuClocksUsingSBLHelper(uint32 Pid)
+{
+	if (!SBLHelper.IsValid())
+	{
+		return false;
+	}
+
+	// Try to connect to the SBLHelper server if we haven't already
+	if (!SBLHelper->IsConnected())
+	{
+		FSBLHelperClient::FConnectionParams ConnectionParams;
+
+		uint16 Port = 8010; // Default tcp port
+
+		// Apply command line port number override, if present
+		{
+			static uint16 CmdLinePortOverride = 0;
+			static bool bCmdLinePortOverrideParsed = false;
+			static bool bCmdLinePortOverrideValid = false;
+
+			if (!bCmdLinePortOverrideParsed)
+			{
+				bCmdLinePortOverrideParsed = true;
+
+				bCmdLinePortOverrideValid = FParse::Value(FCommandLine::Get(), TEXT("sblhport="), CmdLinePortOverride);
+			}
+
+			if (bCmdLinePortOverrideValid)
+			{
+				Port = CmdLinePortOverride;
+			}
+		}
+
+		const FString HostName = FString::Printf(TEXT("localhost:%d"), Port);
+		FIPv4Endpoint::FromHostAndPort(*HostName, ConnectionParams.Endpoint);
+
+		SBLHelper->Connect(ConnectionParams);
+	}
+
+	if (!SBLHelper->IsConnected())
+	{
+		return false;
+	}
+
+	const bool bSentMessage = SBLHelper->LockGpuClock(Pid);
+
+	if (!bSentMessage)
+	{
+		UE_LOG(LogSwitchboard, Error, TEXT("Failed to send message to SBLHelper server to request gpu clock locking"));
+	}
+
+	// We disconnect right away because launches happen only far and in between.
+	SBLHelper->Disconnect();
+
+	return true;
+}
+
diff --git a/Engine/Source/Programs/SwitchboardListener/SblCore/Private/GpuClockerManager.h b/Engine/Source/Programs/SwitchboardListener/SblCore/Private/GpuClockerManager.h
new file mode 100644
--- /dev/null
+++ b/Engine/Source/Programs/SwitchboardListener/SblCore/Private/GpuClockerManager.h
@@ -0,0 +1,56 @@
+// Copyright Epic Games, Inc. All Rights Reserved.
+
+#pragma once
+
+#include "CoreMinimal.h"
+
+
+class FSBLHelperClient;
+
+
+/**
+ * Manages the GpuClocker, whether using local functionality when SBL is launched with elevanted privileges
+ * (which is an NVML requirement to lock Gpu clocks), or the SwitchboardListenerHelper process
+ */
+class FGpuClockerManager
+{
+public:
+
+	FGpuClockerManager();
+	~FGpuClockerManager();
+
+	/** Call periodically to perform service maintenance tasks. */
+	void Tick();
+
+	/** Locks the Gpu clocks for at least the lifetime of the given process id. */
+	bool LockGpuClocksForPid(uint32 Pid);
+
+	/** 
+	 * Must be called when a Pid that this gpu clock manager may have in its internal list has ended. 
+	 * When all pids that conditio the gpu clocks end, then the clocks are put back to their normal state.
+	 */
+	void PidEnded(uint32 Pid);
+
+private:
+
+	/** Requests the SBL Helper executable to lock the Gpu clocks during the lifetime of at least the given process id. */
+	bool LockGpuClocksUsingSBLHelper(uint32 Pid);
+
+private:
+
+	/** Client interface to the Switchboard Listener Helper external process */
+	TSharedPtr<FSBLHelperClient> SBLHelper;
+
+	/** 
+	 * Set of Pids that are keeping the gpus clocks locked. Only applicable when locally
+	 * managed and not when using the external SBLHelper.
+	 */
+	TSet<uint32> LockingPids;
+
+	/** 
+	 * True when we are locally managing the locking of the Gpu clocks. It doesn't imply that the local
+	 * locking succeeded, only that we are not using SBL Helper to do it for the current pids in the set.
+	 */
+	bool bLockManagedLocally = false;
+};
+
diff --git a/Engine/Source/Programs/SwitchboardListener/SblCore/Private/ScopedNvApi.cpp b/Engine/Source/Programs/SwitchboardListener/SblCore/Private/ScopedNvApi.cpp
--- a/Engine/Source/Programs/SwitchboardListener/SblCore/Private/ScopedNvApi.cpp
+++ b/Engine/Source/Programs/SwitchboardListener/SblCore/Private/ScopedNvApi.cpp
@@ -84,7 +84,7 @@ FScopedNvApi::FScopedNvApi()
 
 	if (Statics::InstanceCount == 1)
 	{
-		UE_LOG(LogSwitchboard, Verbose, TEXT("FScopedNvApi InstanceCount incremented to 1"));
+		UE_LOG(LogSwitchboard, VeryVerbose, TEXT("FScopedNvApi InstanceCount incremented to 1"));
 	}
 
 	// Let others know that a new instance of this class was instantiated.
@@ -104,7 +104,7 @@ FScopedNvApi::~FScopedNvApi()
 	// a future instance tries to get a gpu count or read gpu stats.
 	if (!Statics::InstanceCount)
 	{
-		UE_LOG(LogSwitchboard, Verbose, TEXT("FScopedNvApi InstanceCount decremented to 0"));
+		UE_LOG(LogSwitchboard, VeryVerbose, TEXT("FScopedNvApi InstanceCount decremented to 0"));
 		Statics::PhysicalGpuHandles.Reset(); // Keep the memory.
 	}
 
diff --git a/Engine/Source/Programs/SwitchboardListener/SblCore/Private/SwitchboardListener.cpp b/Engine/Source/Programs/SwitchboardListener/SblCore/Private/SwitchboardListener.cpp
--- a/Engine/Source/Programs/SwitchboardListener/SblCore/Private/SwitchboardListener.cpp
+++ b/Engine/Source/Programs/SwitchboardListener/SblCore/Private/SwitchboardListener.cpp
@@ -3,7 +3,7 @@
 #include "SwitchboardListener.h"
 
 #include "CpuUtilizationMonitor.h"
-#include "SBLHelperClient.h"
+#include "GpuClockerManager.h"
 #include "SwitchboardAuth.h"
 #include "SwitchboardListenerApp.h"
 #include "SwitchboardMessageFuture.h"
@@ -50,7 +50,7 @@
 #endif
 
 
-#define QUIC_ENSURE(X)		ensure(QUIC_SUCCEEDED(X))
+#define QUIC_ENSURE(X) ensure(QUIC_SUCCEEDED(X))
 
 
 const FIPv4Endpoint FSwitchboardListener::InvalidEndpoint(FIPv4Address::LanBroadcast, 0);
@@ -370,7 +370,7 @@ FString FSwitchboardCommandLineOptions::ToString(bool bIncludeRedeploy /* = fals
 FSwitchboardListener::FSwitchboardListener(const FSwitchboardCommandLineOptions& InOptions)
 	: Options(InOptions)
 	, CpuMonitor(MakeShared<FCpuUtilizationMonitor>())
-	, SBLHelper(MakeShared<FSBLHelperClient>())
+	, GpuClockerManager(MakeShared<FGpuClockerManager>())
 	, CachedMosaicToposLock(MakeShared<FRWLock>())
 	, CachedMosaicTopos(MakeShared<TArray<FMosaicTopo>>())
 {
@@ -619,7 +619,7 @@ void FSwitchboardListener::Tick()
 	HandleRunningProcesses(RunningProcesses, true);
 	HandleRunningProcesses(FlipModeMonitors, false);
 	SendMessageFutures();
-	SBLHelper->Tick();
+	GpuClockerManager->Tick();
 
 #if PLATFORM_WINDOWS
 	FWorkaroundForHitchingAfterHours::Get().Tick();
@@ -1042,58 +1042,9 @@ bool FSwitchboardListener::Task_StartProcess(const FSwitchboardStartTask& InRunT
 	}
 
 	// Lock Gpu Clocks for the lifetime of this PID, if requested
-	if (InRunTask.bLockGpuClock && SBLHelper.IsValid())
+	if (InRunTask.bLockGpuClock && GpuClockerManager.IsValid())
 	{
-		// Try to connect to the SBLHelper server if we haven't already
-		if (!SBLHelper->IsConnected())
-		{
-			FSBLHelperClient::FConnectionParams ConnectionParams;
-
-			uint16 Port = 8010; // Default tcp port
-
-			// Apply command line port number override, if present
-			{
-				static uint16 CmdLinePortOverride = 0;
-				static bool bCmdLinePortOverrideParsed = false;
-				static bool bCmdLinePortOverrideValid = false;
-
-				if (!bCmdLinePortOverrideParsed)
-				{
-					bCmdLinePortOverrideParsed = true;
-
-					bCmdLinePortOverrideValid = FParse::Value(FCommandLine::Get(), TEXT("sblhport="), CmdLinePortOverride);
-				}
-
-				if (bCmdLinePortOverrideValid)
-				{
-					Port = CmdLinePortOverride;
-				}
-			}
-
-			const FString HostName = FString::Printf(TEXT("localhost:%d"), Port);
-			FIPv4Endpoint::FromHostAndPort(*HostName, ConnectionParams.Endpoint);
-
-			SBLHelper->Connect(ConnectionParams);
-		}
-
-		if (SBLHelper->IsConnected())
-		{
-			const bool bSentMessage = SBLHelper->LockGpuClock(NewProcess->PID);
-
-			if (!bSentMessage)
-			{
-				UE_LOG(LogSwitchboard, Error, TEXT("Failed to send message to SBLHelper server to request gpu clock locking"));
-			}
-
-			// We disconnect right away because launches happen only far and in between.
-			SBLHelper->Disconnect();
-		}
-		else
-		{
-			UE_LOG(LogSwitchboard, Warning, TEXT("Lock Gpu clocks was requested but could not connect to SwitchboardListenerHelper process. "
-				"Please verify that it is running as admin (elevated privileges are required to lock Gpu clocks). "
-			    "If locking Gpu clocks is not desired, this option can be disabled in Switchboard."));
-		}
+		GpuClockerManager->LockGpuClocksForPid(NewProcess->PID);
 	}
 
 	UE_LOG(LogSwitchboard, Display, TEXT("Started process %d: %s %s"), NewProcess->PID, *InRunTask.Command, *InRunTask.Arguments);
@@ -2529,6 +2480,12 @@ void FSwitchboardListener::HandleRunningProcesses(TArray<TSharedPtr<FRunningProc
 					UE_LOG(LogSwitchboard, Display, TEXT("Output:\n%.*s"), StdoutConv.Length(), StdoutConv.Get());
 				}
 
+				// Let the clock manager know that the Pid ended.
+				if (GpuClockerManager.IsValid())
+				{
+					GpuClockerManager->PidEnded(Process->PID);
+				}
+
 				// Notify remote client, which implies that this is a program managed by it.
 				if (bNotifyThatProgramEnded)
 				{
diff --git a/Engine/Source/Programs/SwitchboardListener/SblCore/SblCore.Build.cs b/Engine/Source/Programs/SwitchboardListener/SblCore/SblCore.Build.cs
--- a/Engine/Source/Programs/SwitchboardListener/SblCore/SblCore.Build.cs
+++ b/Engine/Source/Programs/SwitchboardListener/SblCore/SblCore.Build.cs
@@ -65,14 +65,16 @@ public class SblCore : ModuleRules
 		if (Target.Platform == UnrealTargetPlatform.Win64)
 		{
 			PublicAdditionalLibraries.Add(Path.Combine(ModuleDirectory, "NvmlWrapper", "Lib", "Windows", "NvmlWrapper.lib"));
-
-			PublicIncludePaths.AddRange(
-				new string[]
-				{
-				Path.Combine(ModuleDirectory, "NvmlWrapper", "Public"),
-				}
-			);
 		}
+		else if (Target.Platform == UnrealTargetPlatform.Linux)
+		{
+			PublicAdditionalLibraries.Add(Path.Combine(ModuleDirectory, "NvmlWrapper", "Lib", "Linux", "libNvmlWrapper.a"));
+		}
+
+		PublicIncludePaths.AddRange(new string[] 
+		{
+			Path.Combine(ModuleDirectory, "NvmlWrapper", "Public"),
+		});
 
 	}
 }
