SwitchboardListener: Don't warn about CRYPTO_secure_malloc_init failure on Windows.

Stepping into the disassembly yields "xor eax,eax; ret" - our OpenSSL binaries must have been compiled with OPENSSL_NO_SECURE_MEMORY.

diff --git a/Engine/Source/Programs/SwitchboardListener/SblCore/Private/SwitchboardAuth.cpp b/Engine/Source/Programs/SwitchboardListener/SblCore/Private/SwitchboardAuth.cpp
--- a/Engine/Source/Programs/SwitchboardListener/SblCore/Private/SwitchboardAuth.cpp
+++ b/Engine/Source/Programs/SwitchboardListener/SblCore/Private/SwitchboardAuth.cpp
@@ -650,6 +650,8 @@ bool FSwitchboardAuthHelper::Initialize(const FSwitchboardAuthHelper::FSettings&
 	// Mutable copy.
 	FSettings Settings = InSettings;
 
+	// At the time of writing, our Windows OpenSSL binaries appear to be built with OPENSSL_NO_SECURE_MEMORY.
+#if !PLATFORM_WINDOWS
 	// Try to ensure OpenSSL secure heap is initialized.
 	// Not strictly required, but without it, credentials may
 	// get swapped to disk, be included in core dumps, etc.
@@ -666,6 +668,7 @@ bool FSwitchboardAuthHelper::Initialize(const FSwitchboardAuthHelper::FSettings&
 			bSecureHeapNeedsShutdown = true;
 		}
 	}
+#endif // #if !PLATFORM_WINDOWS
 
 	// If the user didn't specify key pair, check the default path,
 	// and possibly generate a new self-signed key pair.
