SBL: Get rid of large number of redundant FlipHistoryMode entries in status messages.

Also got rid of "PresentMode" as flip mode which was causing false glitch messages in the UI.
And added client ip to sending message logs which will help when debugging logs from clients.

#rnx

diff --git a/Engine/Source/Programs/SwitchboardListener/SblCore/Private/SwitchboardListener.cpp b/Engine/Source/Programs/SwitchboardListener/SblCore/Private/SwitchboardListener.cpp
--- a/Engine/Source/Programs/SwitchboardListener/SblCore/Private/SwitchboardListener.cpp
+++ b/Engine/Source/Programs/SwitchboardListener/SblCore/Private/SwitchboardListener.cpp
@@ -2020,7 +2020,17 @@ static void FillOutFlipMode(FSyncStatus& SyncStatus, FRunningProcess* FlipModeMo
 
 		const int32 PresentMonIdx = 11;
 
-		SyncStatus.FlipModeHistory.Add(Fields[PresentMonIdx]); // The first one will be "PresentMode". This is ok. 
+		// Skip "PresentMode", which is not a real flip mode.
+		if (Fields[PresentMonIdx] == TEXT("PresentMode"))
+		{
+			continue;
+		}
+
+		// Only add if it's different from the last entry
+		if ((SyncStatus.FlipModeHistory.Num() == 0) || (SyncStatus.FlipModeHistory.Last() != Fields[PresentMonIdx]))
+		{
+			SyncStatus.FlipModeHistory.Add(Fields[PresentMonIdx]);
+		}
 	}
 }
 #endif // PLATFORM_WINDOWS
@@ -2555,7 +2565,7 @@ bool FSwitchboardListener::SendMessage(const FString& InMessage, const FIPv4Endp
 
 	FConnectionRef Connection = ConnectionsByEndpoint[InEndpoint];
 
-	UE_LOG(LogSwitchboardProtocol, Verbose, TEXT("Sending message %s"), *InMessage);
+	UE_LOG(LogSwitchboardProtocol, Verbose, TEXT("Sending message to client '%s' : %s"), *InEndpoint.ToString(), *InMessage);
 
 	uint64 Utf8Length = FPlatformString::ConvertedLength<UTF8CHAR>(*InMessage, InMessage.Len() + 1);
 	FByteArrayRef SendArray = MakeShared<TArray<uint8>>();
