SBL: Fix internal include restriction.

#rnx

diff --git a/Engine/Source/Programs/SwitchboardListener/SblSlate/SblSlate.Build.cs b/Engine/Source/Programs/SwitchboardListener/SblSlate/SblSlate.Build.cs
--- a/Engine/Source/Programs/SwitchboardListener/SblSlate/SblSlate.Build.cs
+++ b/Engine/Source/Programs/SwitchboardListener/SblSlate/SblSlate.Build.cs
@@ -9,6 +9,8 @@ public class SblSlate : ModuleRules
 	{
 		PublicDefinitions.Add("SBL_SLATE=1");
 
+		PublicIncludePaths.Add(System.IO.Path.Combine(GetModuleDirectory("SwitchboardCommon"), "Internal"));
+
 		PublicDependencyModuleNames.AddRange(
 			new string[]
 			{
