SwitchboardListener: Fix for SwitchboardProtocol log category showing up by default.

#rnx

diff --git a/Engine/Source/Programs/SwitchboardListener/SblCore/Internal/SwitchboardListenerApp.h b/Engine/Source/Programs/SwitchboardListener/SblCore/Internal/SwitchboardListenerApp.h
--- a/Engine/Source/Programs/SwitchboardListener/SblCore/Internal/SwitchboardListenerApp.h
+++ b/Engine/Source/Programs/SwitchboardListener/SblCore/Internal/SwitchboardListenerApp.h
@@ -6,6 +6,7 @@
 
 
 DECLARE_LOG_CATEGORY_EXTERN(LogSwitchboard, Verbose, All);
+DECLARE_LOG_CATEGORY_EXTERN(LogSwitchboardProtocol, Verbose, All);
 
 
 #if !defined(SBL_SLATE)
diff --git a/Engine/Source/Programs/SwitchboardListener/SblCore/Private/SwitchboardProtocol.h b/Engine/Source/Programs/SwitchboardListener/SblCore/Private/SwitchboardProtocol.h
--- a/Engine/Source/Programs/SwitchboardListener/SblCore/Private/SwitchboardProtocol.h
+++ b/Engine/Source/Programs/SwitchboardListener/SblCore/Private/SwitchboardProtocol.h
@@ -6,6 +6,7 @@
 #include "Interfaces/IPv4/IPv4Endpoint.h"
 #include "JsonObjectConverter.h"
 #include "Logging/LogMacros.h"
+#include "SwitchboardListenerApp.h"
 
 
 struct FSwitchboardPacket;
@@ -13,8 +14,6 @@ struct FSwitchboardTask;
 struct FSyncStatus;
 enum class ESyncStatusRequestFlags;
 
-DECLARE_LOG_CATEGORY_EXTERN(LogSwitchboardProtocol, Verbose, All);
-
 
 //~ Messages sent from Listener to Switchboard
 
diff --git a/Engine/Source/Programs/SwitchboardListener/SblSlate/Private/SblMainWindow.cpp b/Engine/Source/Programs/SwitchboardListener/SblSlate/Private/SblMainWindow.cpp
--- a/Engine/Source/Programs/SwitchboardListener/SblSlate/Private/SblMainWindow.cpp
+++ b/Engine/Source/Programs/SwitchboardListener/SblSlate/Private/SblMainWindow.cpp
@@ -549,6 +549,9 @@ TSharedRef<SWidget> FSwitchboardListenerMainWindow::CreateOutputLog()
 		return LogCategory == LogSwitchboard.GetCategoryName();
 	});
 
+	OutputLogCreationParams.DefaultCategorySelection.Emplace(LogSwitchboard.GetCategoryName(), true);
+	OutputLogCreationParams.DefaultCategorySelection.Emplace(LogSwitchboardProtocol.GetCategoryName(), false);
+
 	// Limit logging output to 10,000 so that we don't over consume on memory.
 	OutputLogCreationParams.LoggingLimit = 10000;
 	TSharedRef<SWidget> OutputLog = OutputLogModule.MakeOutputLogWidget(OutputLogCreationParams);
