[SwitchboardListener] Updated the 'taskkill' command calls to also terminate the child processes

#rnx

diff --git a/Engine/Source/Programs/SwitchboardListener/SblCore/Private/SwitchboardListener.cpp b/Engine/Source/Programs/SwitchboardListener/SblCore/Private/SwitchboardListener.cpp
--- a/Engine/Source/Programs/SwitchboardListener/SblCore/Private/SwitchboardListener.cpp
+++ b/Engine/Source/Programs/SwitchboardListener/SblCore/Private/SwitchboardListener.cpp
@@ -1254,7 +1254,7 @@ bool FSwitchboardListener::KillProcessNow(FRunningProcess* InProcess, float Soft
 		// try a soft kill first
 		if(SoftKillTimeout > 0)
 		{
-			const FString Params = FString::Printf(TEXT("/PID %d"), InProcess->PID);
+			const FString Params = FString::Printf(TEXT("/T /F /PID %d"), InProcess->PID);
 
 			FString OutStdOut;
 
