SwitchboardListener: Don't pass deprecated '-dont_restart_as_admin' option to PresentMon

Suppresses a warning from PresentMon about the option being deprecated.

diff --git a/Engine/Source/Programs/SwitchboardListener/SblCore/Private/SwitchboardListener.cpp b/Engine/Source/Programs/SwitchboardListener/SblCore/Private/SwitchboardListener.cpp
--- a/Engine/Source/Programs/SwitchboardListener/SblCore/Private/SwitchboardListener.cpp
+++ b/Engine/Source/Programs/SwitchboardListener/SblCore/Private/SwitchboardListener.cpp
@@ -1924,7 +1924,7 @@ FRunningProcess* FSwitchboardListener::FindOrStartFlipModeMonitorForUUID(const F
 	MonitorProcess->Path = FPaths::EngineDir() / TEXT("Binaries") / TEXT("ThirdParty") / TEXT("PresentMon") / TEXT("Win64") / TEXT("PresentMon-1.8.0-x64.exe");
 
 	FString Arguments = 
-		FString::Printf(TEXT("-session_name session_%d -output_stdout -dont_restart_as_admin -terminate_on_proc_exit -stop_existing_session -process_id %d"), 
+		FString::Printf(TEXT("-session_name session_%d -output_stdout -terminate_on_proc_exit -stop_existing_session -process_id %d"), 
 		Process->PID, Process->PID);
 
 	MonitorProcess->Handle = FPlatformProcess::CreateProc(
