SwitchboardListener: Enable separate handling of stdout and stderr streams.

This enables returning structured data from sbl_helper.py via stdout while continuing to log interactively via stderr.

Also bumps Listener version to 3.2.0.

diff --git a/Engine/Plugins/VirtualProduction/Switchboard/Source/Switchboard/switchboard/devices/unreal/version_helpers.py b/Engine/Plugins/VirtualProduction/Switchboard/Source/Switchboard/switchboard/devices/unreal/version_helpers.py
--- a/Engine/Plugins/VirtualProduction/Switchboard/Source/Switchboard/switchboard/devices/unreal/version_helpers.py
+++ b/Engine/Plugins/VirtualProduction/Switchboard/Source/Switchboard/switchboard/devices/unreal/version_helpers.py
@@ -3,7 +3,7 @@
 import typing
 
 
-LISTENER_COMPATIBLE_VERSION = (3,1)
+LISTENER_COMPATIBLE_VERSION = (3,2)
 
 # Version where redeploy support was initially introduced.
 LISTENER_MIN_REDEPLOY_VERSION = (1,4,0)
diff --git a/Engine/Plugins/VirtualProduction/Switchboard/Source/Switchboard/switchboard/message_protocol.py b/Engine/Plugins/VirtualProduction/Switchboard/Source/Switchboard/switchboard/message_protocol.py
--- a/Engine/Plugins/VirtualProduction/Switchboard/Source/Switchboard/switchboard/message_protocol.py
+++ b/Engine/Plugins/VirtualProduction/Switchboard/Source/Switchboard/switchboard/message_protocol.py
@@ -72,6 +72,7 @@ def create_start_process_message(
     working_dir: str = "",
     *,
     update_clients_with_stdout: bool = False,
+    combined_stdout_stderr: bool = True,
     priority_modifier: int = 0,
     lock_gpu_clock: bool = False,
     hide: bool = False,
@@ -86,6 +87,7 @@ def create_start_process_message(
         'caller': caller,
         'working_dir': working_dir,
         'bUpdateClientsWithStdout': update_clients_with_stdout,
+        'bCombinedStdoutStderr': combined_stdout_stderr,
         'priority_modifier': priority_modifier,
         'bLockGpuClock': lock_gpu_clock,
         'bHide': hide,
diff --git a/Engine/Source/Programs/SwitchboardListener/SblCore/Internal/SwitchboardListenerVersion.h b/Engine/Source/Programs/SwitchboardListener/SblCore/Internal/SwitchboardListenerVersion.h
--- a/Engine/Source/Programs/SwitchboardListener/SblCore/Internal/SwitchboardListenerVersion.h
+++ b/Engine/Source/Programs/SwitchboardListener/SblCore/Internal/SwitchboardListenerVersion.h
@@ -3,5 +3,5 @@
 #pragma once
 
 #define SBLISTENER_VERSION_MAJOR 3
-#define SBLISTENER_VERSION_MINOR 1
-#define SBLISTENER_VERSION_PATCH 1
+#define SBLISTENER_VERSION_MINOR 2
+#define SBLISTENER_VERSION_PATCH 0
diff --git a/Engine/Source/Programs/SwitchboardListener/SblCore/Private/SwitchboardListener.cpp b/Engine/Source/Programs/SwitchboardListener/SblCore/Private/SwitchboardListener.cpp
--- a/Engine/Source/Programs/SwitchboardListener/SblCore/Private/SwitchboardListener.cpp
+++ b/Engine/Source/Programs/SwitchboardListener/SblCore/Private/SwitchboardListener.cpp
@@ -247,10 +247,13 @@ struct FRunningProcess
 
 	void* StdoutParentReadPipe = nullptr;
 	void* StdoutChildWritePipe = nullptr;
+	void* StderrParentReadPipe = nullptr;
+	void* StderrChildWritePipe = nullptr;
 	void* StdinChildReadPipe = nullptr;
 	void* StdinParentWritePipe = nullptr;
 
-	TArray<uint8> Output;
+	TArray<uint8> Stdout;
+	TArray<uint8> Stderr;
 
 	FIPv4Endpoint Recipient;
 	FString Path;
@@ -259,6 +262,7 @@ struct FRunningProcess
 
 	std::atomic<bool> bPendingKill = false;
 	bool bUpdateClientsWithStdout = false;
+	bool bCombinedStdoutStderr = true;
 	bool bLockGpuClock = false;
 
 	bool CreatePipes()
@@ -269,6 +273,15 @@ struct FRunningProcess
 			return false;
 		}
 
+		if (!bCombinedStdoutStderr)
+		{
+			if (!FPlatformProcess::CreatePipe(StderrParentReadPipe, StderrChildWritePipe))
+			{
+				UE_LOG(LogSwitchboard, Error, TEXT("Could not create pipe to read process stderr!"));
+				return false;
+			}
+		}
+
 		if (!FPlatformProcess::CreatePipe(StdinChildReadPipe, StdinParentWritePipe, true))
 		{
 			UE_LOG(LogSwitchboard, Error, TEXT("Could not create pipe to write process input!"));
@@ -281,6 +294,7 @@ struct FRunningProcess
 	void ClosePipes()
 	{
 		FPlatformProcess::ClosePipe(StdoutParentReadPipe, StdoutChildWritePipe);
+		FPlatformProcess::ClosePipe(StderrParentReadPipe, StderrChildWritePipe);
 		FPlatformProcess::ClosePipe(StdinChildReadPipe, StdinParentWritePipe);
 	}
 };
@@ -986,6 +1000,7 @@ bool FSwitchboardListener::Task_StartProcess(const FSwitchboardStartTask& InRunT
 	NewProcess->Name = InRunTask.Name;
 	NewProcess->Caller = InRunTask.Caller;
 	NewProcess->bUpdateClientsWithStdout = InRunTask.bUpdateClientsWithStdout;
+	NewProcess->bCombinedStdoutStderr = InRunTask.bCombinedStdoutStderr;
 	NewProcess->bLockGpuClock = InRunTask.bLockGpuClock;
 	NewProcess->UUID = InRunTask.TaskID; // Process ID is the same as the message ID.
 	NewProcess->PID = 0; // default value
@@ -1011,7 +1026,8 @@ bool FSwitchboardListener::Task_StartProcess(const FSwitchboardStartTask& InRunT
 		PriorityModifier, 
 		WorkingDirectory, 
 		NewProcess->StdoutChildWritePipe,
-		NewProcess->StdinChildReadPipe
+		NewProcess->StdinChildReadPipe,
+		InRunTask.bCombinedStdoutStderr ? NewProcess->StdoutChildWritePipe : NewProcess->StderrChildWritePipe
 	);
 
 	if (!NewProcess->Handle.IsValid())
@@ -1985,11 +2001,11 @@ static void FillOutFlipMode(FSyncStatus& SyncStatus, FRunningProcess* FlipModeMo
 	}
 
 	// Get stdout.
-	FUTF8ToTCHAR StdoutConv(reinterpret_cast<const ANSICHAR*>(FlipModeMonitor->Output.GetData()), FlipModeMonitor->Output.Num());
+	FUTF8ToTCHAR StdoutConv(reinterpret_cast<const ANSICHAR*>(FlipModeMonitor->Stdout.GetData()), FlipModeMonitor->Stdout.Num());
 	const FString StdOut(StdoutConv.Length(), StdoutConv.Get());
 
 	// Clear out the StdOut array.
-	FlipModeMonitor->Output.Empty();
+	FlipModeMonitor->Stdout.Empty();
 
 	// Split into lines
 
@@ -2426,14 +2442,22 @@ void FSwitchboardListener::HandleStdout(const TSharedPtr<FRunningProcess>& Proce
 {
 	SWITCHBOARD_TRACE_CPUPROFILER_EVENT_SCOPE(FSwitchboardListener::HandleStdout);
 
-	TArray<uint8> Output;
-	if (FPlatformProcess::ReadPipeToArray(Process->StdoutParentReadPipe, Output))
+	TArray<uint8> Stdout, Stderr;
+	if (FPlatformProcess::ReadPipeToArray(Process->StdoutParentReadPipe, Stdout))
 	{
-		Process->Output.Append(Output);
+		Process->Stdout.Append(Stdout);
+	}
+
+	if (!Process->bCombinedStdoutStderr)
+	{
+		if (FPlatformProcess::ReadPipeToArray(Process->StderrParentReadPipe, Stderr))
+		{
+			Process->Stderr.Append(Stderr);
+		}
 	}
 
-	// If there was a new stdout, update the clients
-	if (Output.Num() && Process->bUpdateClientsWithStdout)
+	// If there was new output, update the clients
+	if (Process->bUpdateClientsWithStdout && (Stdout.Num() || Stderr.Num()))
 	{
 		FSwitchboardProgramStdout Packet;
 
@@ -2443,7 +2467,15 @@ void FSwitchboardListener::HandleStdout(const TSharedPtr<FRunningProcess>& Proce
 		Packet.Process.Caller = Process->Caller;
 		Packet.Process.Pid = Process->PID;
 
-		Packet.PartialStdoutB64 = FBase64::Encode(Output);
+		if (Stdout.Num())
+		{
+			Packet.PartialStdoutB64 = FBase64::Encode(Stdout);
+		}
+
+		if (Stderr.Num())
+		{
+			Packet.PartialStderrB64 = FBase64::Encode(Stderr);
+		}
 
 		UE::TUniqueLock<UE::FRecursiveMutex> ConnectionsScopeLock(ConnectionsLock);
 		for (const TPair<FIPv4Endpoint, FConnectionRef>& Connection : ConnectionsByEndpoint)
@@ -2484,10 +2516,19 @@ void FSwitchboardListener::HandleRunningProcesses(TArray<TSharedPtr<FRunningProc
 				UE_LOG(LogSwitchboard, Display, TEXT("Process %d (%s) exited with returncode: %d"),
 					Process->PID, *Process->Name, ReturnCode);
 
-				if (ReturnCode != 0 && Process->Output.Num() > 0)
+				if (ReturnCode != 0)
 				{
-					FUTF8ToTCHAR StdoutConv(reinterpret_cast<const ANSICHAR*>(Process->Output.GetData()), Process->Output.Num());
-					UE_LOG(LogSwitchboard, Display, TEXT("Output:\n%.*s"), StdoutConv.Length(), StdoutConv.Get());
+					if (Process->Stdout.Num() > 0)
+					{
+						FUTF8ToTCHAR StdoutConv(reinterpret_cast<const ANSICHAR*>(Process->Stdout.GetData()), Process->Stdout.Num());
+						UE_LOG(LogSwitchboard, Display, TEXT("Stdout:\n%.*s"), StdoutConv.Length(), StdoutConv.Get());
+					}
+
+					if (Process->Stderr.Num() > 0)
+					{
+						FUTF8ToTCHAR StderrConv(reinterpret_cast<const ANSICHAR*>(Process->Stderr.GetData()), Process->Stderr.Num());
+						UE_LOG(LogSwitchboard, Display, TEXT("Stderr:\n%.*s"), StderrConv.Length(), StderrConv.Get());
+					}
 				}
 
 				// Let the clock manager know that the Pid ended.
@@ -2507,7 +2548,8 @@ void FSwitchboardListener::HandleRunningProcesses(TArray<TSharedPtr<FRunningProc
 					Packet.Process.Caller = Process->Caller;
 					Packet.Process.Pid = Process->PID;
 					Packet.Returncode = ReturnCode;
-					Packet.StdoutB64 = FBase64::Encode(Process->Output);
+					Packet.StdoutB64 = FBase64::Encode(Process->Stdout);
+					Packet.StderrB64 = FBase64::Encode(Process->Stderr);
 
 					{
 						UE::TUniqueLock<UE::FRecursiveMutex> ConnectionsScopeLock(ConnectionsLock);
diff --git a/Engine/Source/Programs/SwitchboardListener/SblCore/Private/SwitchboardPacket.h b/Engine/Source/Programs/SwitchboardListener/SblCore/Private/SwitchboardPacket.h
--- a/Engine/Source/Programs/SwitchboardListener/SblCore/Private/SwitchboardPacket.h
+++ b/Engine/Source/Programs/SwitchboardListener/SblCore/Private/SwitchboardPacket.h
@@ -113,6 +113,9 @@ struct FSwitchboardProgramStdout : public FSwitchboardPacket
 
 	UPROPERTY()
 	FString PartialStdoutB64;
+
+	UPROPERTY()
+	FString PartialStderrB64;
 };
 
 USTRUCT()
@@ -134,6 +137,9 @@ struct FSwitchboardProgramEnded : public FSwitchboardPacket
 
 	UPROPERTY()
 	FString StdoutB64;
+
+	UPROPERTY()
+	FString StderrB64;
 };
 
 USTRUCT()
diff --git a/Engine/Source/Programs/SwitchboardListener/SblCore/Private/SwitchboardProtocol.cpp b/Engine/Source/Programs/SwitchboardListener/SblCore/Private/SwitchboardProtocol.cpp
--- a/Engine/Source/Programs/SwitchboardListener/SblCore/Private/SwitchboardProtocol.cpp
+++ b/Engine/Source/Programs/SwitchboardListener/SblCore/Private/SwitchboardProtocol.cpp
@@ -272,6 +272,11 @@ FCreateTaskResult CreateTaskFromCommand(const FString& InCommand, const FIPv4End
 			UpdateClientsWithStdoutField->TryGetBool(Task->bUpdateClientsWithStdout);
 		}
 
+		if (TSharedPtr<FJsonValue> CombinedStdoutStderrField = JsonData->TryGetField(TEXT("bCombinedStdoutStderr")))
+		{
+			CombinedStdoutStderrField->TryGetBool(Task->bCombinedStdoutStderr);
+		}
+
 		if (TSharedPtr<FJsonValue> LockGpuClock = JsonData->TryGetField(TEXT("bLockGpuClock")))
 		{
 			LockGpuClock->TryGetBool(Task->bLockGpuClock);
diff --git a/Engine/Source/Programs/SwitchboardListener/SblCore/Private/SwitchboardTasks.h b/Engine/Source/Programs/SwitchboardListener/SblCore/Private/SwitchboardTasks.h
--- a/Engine/Source/Programs/SwitchboardListener/SblCore/Private/SwitchboardTasks.h
+++ b/Engine/Source/Programs/SwitchboardListener/SblCore/Private/SwitchboardTasks.h
@@ -166,6 +166,7 @@ struct FSwitchboardStartTask : public FSwitchboardTask
 	FString Caller;
 	FString WorkingDir;
 	bool bUpdateClientsWithStdout = false;
+	bool bCombinedStdoutStderr = true;
 	bool bLockGpuClock = false;
 	int32 PriorityModifier = 0;
 	bool bHide = false;
