SwitchboardListener: Prevent child processes from inheriting unintended handles.

Unintentionally leaking various resources into child processes that outlive SBL is suspected to be involved a multitude of intermittent issues, such as:
- Inability to bind to the configured listening port when re-launching SBL.
- NVAPI initialization failures.
- Network performance degradation and disconnects.

This change introduces an "FSwitchboardProcess" platform abstraction that provides two things:
- An API-compatible alternative to FPlatformProcess::CreateProc with more restrictive handle/descriptor inheritance (currently only implemented for Windows).
- An experimental API evolution that will permit future planned additions, such as influencing the environment of spawned processes, without piling onto the existing 11 parameter signature.

The plan of record is to iterate on these ideas in a "fork" specific to Switchboard, and eventually discuss them for potential "upstreaming" into the core platform abstraction.

Also bumps SwitchboardListener version to 3.2.1.

diff --git a/Engine/Source/Programs/SwitchboardListener/SblCore/Internal/SwitchboardListenerVersion.h b/Engine/Source/Programs/SwitchboardListener/SblCore/Internal/SwitchboardListenerVersion.h
--- a/Engine/Source/Programs/SwitchboardListener/SblCore/Internal/SwitchboardListenerVersion.h
+++ b/Engine/Source/Programs/SwitchboardListener/SblCore/Internal/SwitchboardListenerVersion.h
@@ -4,4 +4,4 @@
 
 #define SBLISTENER_VERSION_MAJOR 3
 #define SBLISTENER_VERSION_MINOR 2
-#define SBLISTENER_VERSION_PATCH 0
+#define SBLISTENER_VERSION_PATCH 1
diff --git a/Engine/Source/Programs/SwitchboardListener/SblCore/Internal/SwitchboardProcess.h b/Engine/Source/Programs/SwitchboardListener/SblCore/Internal/SwitchboardProcess.h
new file mode 100644
--- /dev/null
+++ b/Engine/Source/Programs/SwitchboardListener/SblCore/Internal/SwitchboardProcess.h
@@ -0,0 +1,111 @@
+// Copyright Epic Games, Inc. All Rights Reserved.
+
+#pragma once
+
+#include "Containers/ArrayView.h"
+#include "Containers/StringView.h"
+#include "HAL/PlatformProcess.h"
+#include "Misc/Optional.h"
+
+
+/**
+ * Alternative implementation of FPlatformProcess::CreateProc with some additional features.
+ * Currently, the primary difference is that handles must be explicitly marked for inheritance.
+ */
+struct FGenericSwitchboardProcess
+{
+public:
+	struct FCreateProcParams
+	{
+		/** Path to the executable to launch. */
+		FStringView Executable;
+
+		/** Command line arguments. */
+		TArrayView<const FStringView> Args;
+
+		/** If true, the new process will have its own window. */
+		bool bLaunchDetached = false;
+
+		/** If true, the new process will be minimized in the task bar. */
+		bool bLaunchMinimized = false;
+
+		/** If true, the new process will not have a window or be in the task bar. */
+		bool bLaunchReallyHidden = false;
+
+		/** -2 idle, -1 low, 0 normal, 1 high, 2 higher */
+		int32 PriorityModifier = 0;
+
+		/** Working directory in which to start the program. If unset, uses the current working directory. */
+		TOptional<FStringView> WorkingDirectory;
+
+		/** The read end of a pipe which will serve as stdin for the child process. */
+		void* StdinPipe = nullptr;
+
+		/** The write end of a pipe which will serve as stdout for the child process. */
+		void* StdoutPipe = nullptr;
+
+		/** The write end of a pipe which will serve as stderr for the child process. */
+		void* StderrPipe = nullptr;
+
+		/**
+		 * Additional handles which the descendant process is permitted to inherit.
+		 * StdinPipe, StdoutPipe, and StderrPipe are always inheritable, and should not be added to this list.
+		 * Ignored if `bInheritAllHandles == true` (legacy behavior; not recommended).
+		 */
+		TArrayView<void*> AdditionalInheritedHandles;
+
+		/**
+		 * Permit child process to inherit any handles from the parent process not explicitly excluded otherwise.
+		 * This was the legacy behavior if any standard I/O redirection pipes were specified.
+		 * This can lead to arcane resource leaks, and is not recommended.
+		 */
+		bool bInheritAllHandles = false;
+	};
+
+	struct FCreateProcResult
+	{
+		FProcHandle Handle;
+		uint32 ProcessId = 0;
+	};
+
+public:
+	static FCreateProcResult CreateProc(const FCreateProcParams& InParams);
+
+	// FPlatformProcess-compatible signature that calls into our implementation with more restrictive handle inheritance.
+	static FProcHandle CreateProc(
+		const TCHAR* URL,
+		const TCHAR* Parms,
+		bool bLaunchDetached,
+		bool bLaunchHidden,
+		bool bLaunchReallyHidden,
+		uint32* OutProcessID,
+		int32 PriorityModifier,
+		const TCHAR* OptionalWorkingDirectory,
+		void* PipeWriteChild,
+		void* PipeReadChild = nullptr
+	);
+
+	// FPlatformProcess-compatible signature that calls into our implementation with more restrictive handle inheritance.
+	static FProcHandle CreateProc(
+		const TCHAR* URL,
+		const TCHAR* Parms,
+		bool bLaunchDetached,
+		bool bLaunchHidden,
+		bool bLaunchReallyHidden,
+		uint32* OutProcessID,
+		int32 PriorityModifier,
+		const TCHAR* OptionalWorkingDirectory,
+		void* PipeWriteChild,
+		void* PipeReadChild,
+		void* PipeStdErrChild
+	);
+};
+
+
+// TODO: Restricted FD inheritance is achievable with POSIX_SPAWN_CLOEXEC_DEFAULT,
+// and maybe posix_spawn_file_actions_addclosefrom_np?
+#if PLATFORM_WINDOWS
+#include COMPILED_PLATFORM_HEADER(SwitchboardProcess.h)
+#else
+using FSwitchboardProcess = FGenericSwitchboardProcess;
+#endif
diff --git a/Engine/Source/Programs/SwitchboardListener/SblCore/Internal/Windows/WindowsSwitchboardProcess.h b/Engine/Source/Programs/SwitchboardListener/SblCore/Internal/Windows/WindowsSwitchboardProcess.h
new file mode 100644
--- /dev/null
+++ b/Engine/Source/Programs/SwitchboardListener/SblCore/Internal/Windows/WindowsSwitchboardProcess.h
@@ -0,0 +1,18 @@
+// Copyright Epic Games, Inc. All Rights Reserved.
+
+#pragma once
+
+#include "SwitchboardProcess.h"
+
+
+/** Windows implementation of FSwitchboardProcess. See FGenericSwitchboardProcess for details. */
+struct FWindowsSwitchboardProcess : public FGenericSwitchboardProcess
+{
+public:
+	using FGenericSwitchboardProcess::CreateProc;
+
+	static FCreateProcResult CreateProc(const FCreateProcParams& InParams);
+};
+
+
+using FSwitchboardProcess = FWindowsSwitchboardProcess;
diff --git a/Engine/Source/Programs/SwitchboardListener/SblCore/Private/GenericSwitchboardProcess.cpp b/Engine/Source/Programs/SwitchboardListener/SblCore/Private/GenericSwitchboardProcess.cpp
new file mode 100644
--- /dev/null
+++ b/Engine/Source/Programs/SwitchboardListener/SblCore/Private/GenericSwitchboardProcess.cpp
@@ -0,0 +1,91 @@
+// Copyright Epic Games, Inc. All Rights Reserved.
+
+#include "SwitchboardProcess.h"
+
+
+FGenericSwitchboardProcess::FCreateProcResult FGenericSwitchboardProcess::CreateProc(const FCreateProcParams& InParams)
+{
+	FCreateProcResult Result;
+
+	const FString Args = FString::Join(InParams.Args, TEXT(" "));
+
+	Result.Handle = FPlatformProcess::CreateProc(
+		*FString(InParams.Executable),
+		*Args,
+		InParams.bLaunchDetached,
+		InParams.bLaunchMinimized,
+		InParams.bLaunchReallyHidden,
+		&Result.ProcessId,
+		InParams.PriorityModifier,
+		InParams.WorkingDirectory ? *FString(*InParams.WorkingDirectory) : nullptr,
+		InParams.StdoutPipe,
+		InParams.StdinPipe,
+		InParams.StderrPipe
+	);
+
+	return Result;
+}
+
+
+
+FProcHandle FGenericSwitchboardProcess::CreateProc(
+	const TCHAR* URL,
+	const TCHAR* Parms,
+	bool bLaunchDetached,
+	bool bLaunchHidden,
+	bool bLaunchReallyHidden,
+	uint32* OutProcessID,
+	int32 PriorityModifier,
+	const TCHAR* OptionalWorkingDirectory,
+	void* PipeWriteChild,
+	void* PipeReadChild
+)
+{
+	return FSwitchboardProcess::CreateProc(URL, Parms, bLaunchDetached, bLaunchHidden, bLaunchReallyHidden, OutProcessID, PriorityModifier, OptionalWorkingDirectory, PipeWriteChild, PipeReadChild, nullptr);
+}
+
+
+FProcHandle FGenericSwitchboardProcess::CreateProc(
+	const TCHAR* URL,
+	const TCHAR* Parms,
+	bool bLaunchDetached,
+	bool bLaunchHidden,
+	bool bLaunchReallyHidden,
+	uint32* OutProcessID,
+	int32 PriorityModifier,
+	const TCHAR* OptionalWorkingDirectory,
+	void* PipeWriteChild,
+	void* PipeReadChild,
+	void* PipeStdErrChild
+)
+{
+	const FString UrlStr(URL);
+	const FString ParmsStr(Parms);
+	const TArray<FStringView> ArgsArray = { FStringView(ParmsStr) };
+	const FString WorkingDirectoryStr(OptionalWorkingDirectory);
+
+	FCreateProcParams Params = {
+		.Executable = UrlStr,
+		.Args = ArgsArray,
+		.bLaunchDetached = bLaunchDetached,
+		.bLaunchMinimized = bLaunchHidden,
+		.bLaunchReallyHidden = bLaunchReallyHidden,
+		.PriorityModifier = PriorityModifier,
+		.StdinPipe = PipeReadChild,
+		.StdoutPipe = PipeWriteChild,
+		.StderrPipe = PipeStdErrChild,
+	};
+
+	if (OptionalWorkingDirectory)
+	{
+		Params.WorkingDirectory = WorkingDirectoryStr;
+	}
+
+	FCreateProcResult Result = FSwitchboardProcess::CreateProc(Params);
+	if (OutProcessID)
+	{
+		*OutProcessID = Result.ProcessId;
+	}
+
+	return Result.Handle;
+}
diff --git a/Engine/Source/Programs/SwitchboardListener/SblCore/Private/SwitchboardListener.cpp b/Engine/Source/Programs/SwitchboardListener/SblCore/Private/SwitchboardListener.cpp
--- a/Engine/Source/Programs/SwitchboardListener/SblCore/Private/SwitchboardListener.cpp
+++ b/Engine/Source/Programs/SwitchboardListener/SblCore/Private/SwitchboardListener.cpp
@@ -8,6 +8,7 @@
 #include "SwitchboardListenerApp.h"
 #include "SwitchboardMessageFuture.h"
 #include "SwitchboardPacket.h"
+#include "SwitchboardProcess.h"
 #include "SwitchboardProtocol.h"
 #include "SwitchboardTasks.h"
 #include "SyncStatus.h"
@@ -1016,7 +1017,7 @@ bool FSwitchboardListener::Task_StartProcess(const FSwitchboardStartTask& InRunT
 	const int32 PriorityModifier = InRunTask.PriorityModifier;
 	const TCHAR* WorkingDirectory = InRunTask.WorkingDir.IsEmpty() ? nullptr : *InRunTask.WorkingDir;
 
-	NewProcess->Handle = FPlatformProcess::CreateProc(
+	NewProcess->Handle = FSwitchboardProcess::CreateProc(
 		*InRunTask.Command, 
 		*InRunTask.Arguments, 
 		bLaunchDetached, 
@@ -1444,7 +1445,7 @@ bool FSwitchboardListener::Task_RedeployListener(const FSwitchboardRedeployListe
 	const int32 PriorityModifier = 0;
 	const TCHAR* WorkingDirectory = nullptr;
 	void* PipeWriteChild = nullptr;
-	FProcHandle RedeployedProcHandle = FPlatformProcess::CreateProc(*TempNewListenerPath, *RedeployArgs, bLaunchDetached, bLaunchHidden, bLaunchReallyHidden, OutProcessId, PriorityModifier, WorkingDirectory, PipeWriteChild);
+	FProcHandle RedeployedProcHandle = FSwitchboardProcess::CreateProc(*TempNewListenerPath, *RedeployArgs, bLaunchDetached, bLaunchHidden, bLaunchReallyHidden, OutProcessId, PriorityModifier, WorkingDirectory, PipeWriteChild);
 	if (!RedeployedProcHandle.IsValid())
 	{
 		const FString ErrorMsg(TEXT("Unable to launch new listener"));
@@ -1943,7 +1944,7 @@ FRunningProcess* FSwitchboardListener::FindOrStartFlipModeMonitorForUUID(const F
 		FString::Printf(TEXT("-session_name session_%d -output_stdout -terminate_on_proc_exit -stop_existing_session -process_id %d"), 
 		Process->PID, Process->PID);
 
-	MonitorProcess->Handle = FPlatformProcess::CreateProc(
+	MonitorProcess->Handle = FSwitchboardProcess::CreateProc(
 		*MonitorProcess->Path,
 		*Arguments,
 		bLaunchDetached,
diff --git a/Engine/Source/Programs/SwitchboardListener/SblCore/Private/Windows/WindowsSwitchboardProcess.cpp b/Engine/Source/Programs/SwitchboardListener/SblCore/Private/Windows/WindowsSwitchboardProcess.cpp
new file mode 100644
--- /dev/null
+++ b/Engine/Source/Programs/SwitchboardListener/SblCore/Private/Windows/WindowsSwitchboardProcess.cpp
@@ -0,0 +1,189 @@
+// Copyright Epic Games, Inc. All Rights Reserved.
+
+#include "Windows/WindowsSwitchboardProcess.h"
+#include "Misc/ScopeExit.h"
+
+#include "Windows/AllowWindowsPlatformTypes.h"
+#	include <processthreadsapi.h>
+#include "Windows/HideWindowsPlatformTypes.h"
+
+
+namespace UE::SwitchboardListener::Private
+{
+	void LogLastError(const TCHAR* ErrorContext)
+	{
+		const DWORD ErrorCode = GetLastError();
+		TCHAR ErrorMessage[512];
+		FWindowsPlatformMisc::GetSystemErrorMessage(ErrorMessage, 512, ErrorCode);
+		UE_LOG(LogWindows, Warning, TEXT("%s failed: %s (0x%08x)"), ErrorContext, ErrorMessage, ErrorCode);
+	}
+}
+
+
+// This is required because nested UpdateProcThreadAttribute macros embed TRUE/FALSE :|
+#include "Windows/AllowWindowsPlatformTypes.h"
+
+FWindowsSwitchboardProcess::FCreateProcResult FWindowsSwitchboardProcess::CreateProc(const FCreateProcParams& InParams)
+{
+	// initialize process creation flags
+	uint32 CreateFlags = NORMAL_PRIORITY_CLASS;
+	if (InParams.PriorityModifier < 0)
+	{
+		CreateFlags = (InParams.PriorityModifier == -1) ? BELOW_NORMAL_PRIORITY_CLASS : IDLE_PRIORITY_CLASS;
+	}
+	else if (InParams.PriorityModifier > 0)
+	{
+		CreateFlags = (InParams.PriorityModifier == 1) ? ABOVE_NORMAL_PRIORITY_CLASS : HIGH_PRIORITY_CLASS;
+	}
+
+	if (InParams.bLaunchDetached)
+	{
+		CreateFlags |= DETACHED_PROCESS;
+	}
+
+	// initialize window flags
+	uint32 dwFlags = 0;
+	uint16 ShowWindowFlags = SW_HIDE;
+	if (InParams.bLaunchReallyHidden)
+	{
+		dwFlags = STARTF_USESHOWWINDOW;
+	}
+	else if (InParams.bLaunchMinimized)
+	{
+		dwFlags = STARTF_USESHOWWINDOW;
+		ShowWindowFlags = SW_SHOWMINNOACTIVE;
+	}
+
+	if (InParams.StdoutPipe != nullptr || InParams.StdinPipe != nullptr || InParams.StderrPipe != nullptr)
+	{
+		dwFlags |= STARTF_USESTDHANDLES;
+	}
+
+	bool bInheritHandles = (dwFlags & STARTF_USESTDHANDLES) != 0;
+	bInheritHandles |= InParams.bInheritAllHandles;
+	bInheritHandles |= InParams.AdditionalInheritedHandles.Num() > 0;
+
+	// initialize startup info
+	CreateFlags |= EXTENDED_STARTUPINFO_PRESENT;
+	STARTUPINFOEXW StartupInfoEx = { 0 };
+	StartupInfoEx.StartupInfo = {
+		sizeof(STARTUPINFOEXW),
+		nullptr,
+		nullptr,
+		nullptr,
+		(DWORD)CW_USEDEFAULT,
+		(DWORD)CW_USEDEFAULT,
+		(DWORD)CW_USEDEFAULT,
+		(DWORD)CW_USEDEFAULT,
+		(DWORD)0,
+		(DWORD)0,
+		(DWORD)0,
+		(DWORD)dwFlags,
+		ShowWindowFlags,
+		0,
+		nullptr,
+		HANDLE(InParams.StdinPipe),
+		HANDLE(InParams.StdoutPipe),
+		HANDLE(InParams.StderrPipe)
+	};
+
+	ON_SCOPE_EXIT
+	{
+		if (StartupInfoEx.lpAttributeList)
+		{
+			DeleteProcThreadAttributeList(StartupInfoEx.lpAttributeList);
+			FMemory::Free(StartupInfoEx.lpAttributeList);
+			StartupInfoEx.lpAttributeList = nullptr;
+		}
+	};
+
+	TArray<HANDLE> InheritableHandles;
+
+	if (bInheritHandles && !InParams.bInheritAllHandles)
+	{
+		InheritableHandles.Reserve(InParams.AdditionalInheritedHandles.Num() + 3);
+
+		InheritableHandles.Append(InParams.AdditionalInheritedHandles);
+
+		if (InParams.StdinPipe)
+		{
+			InheritableHandles.Add(InParams.StdinPipe);
+		}
+
+		if (InParams.StdoutPipe)
+		{
+			InheritableHandles.Add(InParams.StdoutPipe);
+		}
+
+		if (InParams.StderrPipe)
+		{
+			InheritableHandles.Add(InParams.StderrPipe);
+		}
+
+		// This "fails" with ERROR_INSUFFICIENT_BUFFER by design.
+		const int32 NumAttributes = 1;
+		SIZE_T AttrListSize = 0;
+		InitializeProcThreadAttributeList(nullptr, NumAttributes, 0, &AttrListSize);
+
+		const DWORD ErrorCode = GetLastError();
+		if (ErrorCode != ERROR_INSUFFICIENT_BUFFER || AttrListSize <= 0)
+		{
+			UE::SwitchboardListener::Private::LogLastError(TEXT("InitializeProcThreadAttributeList"));
+			return {};
+		}
+
+		StartupInfoEx.lpAttributeList = (LPPROC_THREAD_ATTRIBUTE_LIST)FMemory::Malloc(AttrListSize);
+
+		if (!InitializeProcThreadAttributeList(StartupInfoEx.lpAttributeList, NumAttributes, 0, &AttrListSize))
+		{
+			UE::SwitchboardListener::Private::LogLastError(TEXT("InitializeProcThreadAttributeList"));
+			return {};
+		}
+
+		if (!UpdateProcThreadAttribute(StartupInfoEx.lpAttributeList, 0, PROC_THREAD_ATTRIBUTE_HANDLE_LIST,
+				InheritableHandles.GetData(), InheritableHandles.Num() * InheritableHandles.GetTypeSize(),
+				nullptr, nullptr)
+			)
+		{
+			UE::SwitchboardListener::Private::LogLastError(TEXT("UpdateProcThreadAttribute"));
+			return {};
+		}
+	}
+
+	// create the child process
+	const FString Args = FString::Join(InParams.Args, TEXT(" "));
+	FString CommandLine = FString::Printf(TEXT("\"%.*s\" %s"), InParams.Executable.Len(), InParams.Executable.GetData(), *Args);
+	PROCESS_INFORMATION ProcInfo;
+
+	if (!CreateProcess(
+			NULL,
+			CommandLine.GetCharArray().GetData(),
+			nullptr,
+			nullptr,
+			bInheritHandles,
+			(DWORD)CreateFlags,
+			NULL,
+			InParams.WorkingDirectory ? *FString(*InParams.WorkingDirectory) : nullptr,
+			(LPSTARTUPINFOW)&StartupInfoEx,
+			&ProcInfo
+		))
+	{
+		const DWORD ErrorCode = GetLastError();
+		UE::SwitchboardListener::Private::LogLastError(TEXT("CreateProcess"));
+		if (ErrorCode == ERROR_NOT_ENOUGH_MEMORY || ErrorCode == ERROR_OUTOFMEMORY)
+		{
+			// These errors are common enough that we want some available memory information
+			FPlatformMemoryStats Stats = FPlatformMemory::GetStats();
+			UE_LOG(LogWindows, Warning, TEXT("Mem used: %.2f MB, OS Free %.2f MB"), Stats.UsedPhysical / 1048576.0f, Stats.AvailablePhysical / 1048576.0f);
+		}
+		UE_LOG(LogWindows, Warning, TEXT("URL: %s"), *CommandLine);
+
+		return {};
+	}
+
+	::CloseHandle( ProcInfo.hThread );
+
+	return { FProcHandle(ProcInfo.hProcess), ProcInfo.dwProcessId };
+}
+
+#include "Windows/HideWindowsPlatformTypes.h"
