SwitchboardListener: Avoid adding redundant handle when StderrPipe == StdoutPipe.

#rnx

diff --git a/Engine/Source/Programs/SwitchboardListener/SblCore/Private/Windows/WindowsSwitchboardProcess.cpp b/Engine/Source/Programs/SwitchboardListener/SblCore/Private/Windows/WindowsSwitchboardProcess.cpp
--- a/Engine/Source/Programs/SwitchboardListener/SblCore/Private/Windows/WindowsSwitchboardProcess.cpp
+++ b/Engine/Source/Programs/SwitchboardListener/SblCore/Private/Windows/WindowsSwitchboardProcess.cpp
@@ -115,7 +115,7 @@ FWindowsSwitchboardProcess::FCreateProcResult FWindowsSwitchboardProcess::Create
 			InheritableHandles.Add(InParams.StdoutPipe);
 		}
 
-		if (InParams.StderrPipe)
+		if (InParams.StderrPipe && InParams.StderrPipe != InParams.StdoutPipe)
 		{
 			InheritableHandles.Add(InParams.StderrPipe);
 		}
