SwitchboardListener: Try to avoid leaking trace sessions via PresentMon invocations.

- PresentMon sessions are now assigned deterministic IDs.
- We take steps to ensure the underlying ETW trace session is terminated when PresentMon exits, even in the event of an unclean shutdown.
- When SBL is restarted, we attempt to detect and clean up any leaked trace sessions from previous launches.

Bumps SBL to 3.2.2. Also fixes STARTUPINFOEXW initialization in FWindowsSwitchboardProcess to avoid a clang warning.

diff --git a/Engine/Source/Programs/SwitchboardListener/SblCore/Internal/SwitchboardListener.h b/Engine/Source/Programs/SwitchboardListener/SblCore/Internal/SwitchboardListener.h
--- a/Engine/Source/Programs/SwitchboardListener/SblCore/Internal/SwitchboardListener.h
+++ b/Engine/Source/Programs/SwitchboardListener/SblCore/Internal/SwitchboardListener.h
@@ -6,6 +6,7 @@
 #include "SwitchboardAuth.h"
 #include "Async/Mutex.h"
 #include "Async/RecursiveMutex.h"
+#include "Containers/BitArray.h"
 #include "Containers/Queue.h"
 #include "HAL/PlatformProcess.h"
 #include "Interfaces/IPv4/IPv4Endpoint.h"
@@ -160,7 +161,18 @@ private:
 	bool Task_FreeListenerBinary(const FSwitchboardFreeListenerBinaryTask& InRenameProcessTask);
 
 	bool KillProcessNow(FRunningProcess* InProcess, float SoftKillTimeout = 0.0f);
+
+#if PLATFORM_WINDOWS
+	using FFlipModeMonitorSessionId = uint8;
+	static constexpr FFlipModeMonitorSessionId InvalidFlipModeSessionId = FFlipModeMonitorSessionId(-1);
+	static constexpr FFlipModeMonitorSessionId MaxFlipModeMonitorSessions = 16;
+	void HarvestExistingFlipModeSessions();
+	bool EnsureFlipModeSessionAvailable(FFlipModeMonitorSessionId SessionId, bool* bOutStoppedSession = nullptr);
+	FFlipModeMonitorSessionId AllocFlipModeMonitorSessionId();
+	void FreeFlipModeMonitorSessionId(FFlipModeMonitorSessionId SessionId);
+	FString MakeFlipModeMonitorSessionName(FFlipModeMonitorSessionId SessionId);
 	FRunningProcess* FindOrStartFlipModeMonitorForUUID(const FGuid& UUID);
+#endif
 
 	void CleanUpDisconnectedSockets();
 	void HandleStdout(const TSharedPtr<FRunningProcess>& Process);
@@ -261,6 +273,11 @@ private:
 	TQueue<TUniquePtr<FSwitchboardTask>, EQueueMode::Spsc> DisconnectTasks;
 	TArray<TSharedPtr<FRunningProcess>> RunningProcesses;
 	TArray<TSharedPtr<FRunningProcess>> FlipModeMonitors;
+
+#if PLATFORM_WINDOWS
+	TBitArray<> RunningFlipModeSessions;
+#endif
+
 	TArray<FSwitchboardMessageFuture> MessagesFutures;
 	TSharedPtr<FCpuUtilizationMonitor> CpuMonitor;
 
diff --git a/Engine/Source/Programs/SwitchboardListener/SblCore/Internal/SwitchboardListenerVersion.h b/Engine/Source/Programs/SwitchboardListener/SblCore/Internal/SwitchboardListenerVersion.h
--- a/Engine/Source/Programs/SwitchboardListener/SblCore/Internal/SwitchboardListenerVersion.h
+++ b/Engine/Source/Programs/SwitchboardListener/SblCore/Internal/SwitchboardListenerVersion.h
@@ -4,4 +4,4 @@
 
 #define SBLISTENER_VERSION_MAJOR 3
 #define SBLISTENER_VERSION_MINOR 2
-#define SBLISTENER_VERSION_PATCH 1
+#define SBLISTENER_VERSION_PATCH 2
diff --git a/Engine/Source/Programs/SwitchboardListener/SblCore/Private/SwitchboardListener.cpp b/Engine/Source/Programs/SwitchboardListener/SblCore/Private/SwitchboardListener.cpp
--- a/Engine/Source/Programs/SwitchboardListener/SblCore/Private/SwitchboardListener.cpp
+++ b/Engine/Source/Programs/SwitchboardListener/SblCore/Private/SwitchboardListener.cpp
@@ -42,6 +42,7 @@
 #include <Windows.h>
 #include <WinUser.h>
 #include <shellapi.h>
+#include <evntrace.h>
 #include "Windows/HideWindowsPlatformTypes.h"
 
 #include "ScopedNvApi.h"
@@ -151,7 +152,7 @@ namespace
 			{QUIC_STATUS_VER_NEG_ERROR		 , TEXT("A version negotiation error was encountered.")},
 			{QUIC_STATUS_USER_CANCELED		 , TEXT("The peer app/user canceled the connection during the handshake.")},
 			{QUIC_STATUS_ALPN_NEG_FAILURE	 , TEXT("The connection handshake failed to negotiate a common ALPN.")},
-			{QUIC_STATUS_STREAM_LIMIT_REACHED , TEXT("A stream failed to start because the peer doesn't allow any more to be open at this time.")}
+			{QUIC_STATUS_STREAM_LIMIT_REACHED, TEXT("A stream failed to start because the peer doesn't allow any more to be open at this time.")}
 		};
 		if (FString* ErrorString = StatStringMap.Find(InStatus))
 		{
@@ -161,7 +162,6 @@ namespace
 	}
 
 #if PLATFORM_WINDOWS
-
 	/** 
 	 * Manages an instance of FScopedNvApi to keep the library initialized until it doesn't get used for a while.
 	 * The objective is to work around an issue where NvApi can sometimes cause hitches after it has been initialized for 
@@ -266,6 +266,8 @@ struct FRunningProcess
 	bool bCombinedStdoutStderr = true;
 	bool bLockGpuClock = false;
 
+	FTSSimpleMulticastDelegate OnProcessEnded;
+
 	bool CreatePipes()
 	{
 		if (!FPlatformProcess::CreatePipe(StdoutParentReadPipe, StdoutChildWritePipe))
@@ -563,6 +565,12 @@ bool FSwitchboardListener::StartListening()
 		return false;
 	}
 
+#if PLATFORM_WINDOWS
+	// If we've made it this far, we were able to bind the listen endpoint successfully,
+	// which precludes interfering with any other running instance.
+	HarvestExistingFlipModeSessions();
+#endif
+
 	return true;
 }
 
@@ -1883,6 +1891,114 @@ static void FillOutTaskbarAutoHide(FSyncStatus& SyncStatus)
 
 
 #if PLATFORM_WINDOWS
+void FSwitchboardListener::HarvestExistingFlipModeSessions()
+{
+	RunningFlipModeSessions.Reserve(MaxFlipModeMonitorSessions);
+	RunningFlipModeSessions.Reset();
+
+	for (FFlipModeMonitorSessionId SessionId = 0; SessionId < MaxFlipModeMonitorSessions; ++SessionId)
+	{
+		bool bOutStoppedSession = false;
+		if (!EnsureFlipModeSessionAvailable(SessionId, &bOutStoppedSession))
+		{
+			UE_LOGFMT(LogSwitchboard, Error, "Unable to stop existing flip mode session {SessionId}", SessionId);
+			RunningFlipModeSessions.Add(true, 1);
+		}
+		else
+		{
+			if (bOutStoppedSession)
+			{
+				UE_LOGFMT(LogSwitchboard, Warning, "Stopped existing flip mode session {SessionId}", SessionId);
+			}
+
+			RunningFlipModeSessions.Add(false, 1);
+		}
+	}
+}
+
+
+bool FSwitchboardListener::EnsureFlipModeSessionAvailable(FFlipModeMonitorSessionId SessionId, bool* bOutStoppedSession)
+{
+	const FString SessionName = MakeFlipModeMonitorSessionName(SessionId);
+
+	EVENT_TRACE_PROPERTIES TraceProps;
+	FMemory::Memzero(TraceProps);
+	TraceProps.Wnode.BufferSize = sizeof(TraceProps);
+	TraceProps.Wnode.Flags = WNODE_FLAG_TRACED_GUID;
+
+	if (bOutStoppedSession != nullptr)
+	{
+		*bOutStoppedSession = false;
+	}
+
+	const uint32 Result = ControlTrace(0, *SessionName, &TraceProps, EVENT_TRACE_CONTROL_STOP);
+	if (Result == ERROR_SUCCESS || Result == ERROR_MORE_DATA)
+	{
+		if (bOutStoppedSession != nullptr)
+		{
+			*bOutStoppedSession = true;
+		}
+		return true;
+	}
+	else if (Result == ERROR_WMI_INSTANCE_NOT_FOUND)
+	{
+		// Didn't exist.
+		return true;
+	}
+
+	UE_LOGFMT(LogSwitchboard, Error, "Failed to stop flip mode session {SessionName} with result {Result}",
+		SessionName, Result);
+
+	return false;
+}
+
+
+FSwitchboardListener::FFlipModeMonitorSessionId FSwitchboardListener::AllocFlipModeMonitorSessionId()
+{
+	check(IsInGameThread());
+
+	const int32 AvailableId = RunningFlipModeSessions.FindAndSetFirstZeroBit();
+	if (AvailableId < 0)
+	{
+		UE_LOGFMT(LogSwitchboard, Error, "Failed to allocate flip mode monitor session ID; too many sessions");
+		return InvalidFlipModeSessionId;
+	}
+
+	UE_LOGFMT(LogSwitchboard, Verbose, "Allocated flip mode monitor session ID {ID}", AvailableId);
+
+	return static_cast<FFlipModeMonitorSessionId>(AvailableId);
+}
+
+
+void FSwitchboardListener::FreeFlipModeMonitorSessionId(FFlipModeMonitorSessionId SessionId)
+{
+	check(IsInGameThread());
+
+	if (!ensure(SessionId < RunningFlipModeSessions.Num() && SessionId < MaxFlipModeMonitorSessions))
+	{
+		UE_LOGFMT(LogSwitchboard, Error, "Attempted to free invalid flip mode monitor session ID {Id}", SessionId);
+		return;
+	}
+
+	if (!ensure(RunningFlipModeSessions[SessionId]))
+	{
+		UE_LOGFMT(LogSwitchboard, Error, "Attempted to free unused flip mode monitor session ID {Id}", SessionId);
+		return;
+	}
+
+	RunningFlipModeSessions[SessionId] = false;
+
+	UE_LOGFMT(LogSwitchboard, Verbose, "Freed flip mode monitor session ID {ID}", SessionId);
+}
+
+
+FString FSwitchboardListener::MakeFlipModeMonitorSessionName(FFlipModeMonitorSessionId SessionId)
+{
+	return FString::Printf(TEXT("switchboard_%s_presentmon_%d"),
+		 *ListenerEndpoint->ToString(), SessionId);
+}
+
+
 FRunningProcess* FSwitchboardListener::FindOrStartFlipModeMonitorForUUID(const FGuid& UUID)
 {
 	SWITCHBOARD_TRACE_CPUPROFILER_EVENT_SCOPE(FSwitchboardListener::FindOrStartFlipModeMonitorForUUID);
@@ -1940,9 +2056,32 @@ FRunningProcess* FSwitchboardListener::FindOrStartFlipModeMonitorForUUID(const F
 
 	MonitorProcess->Path = FPaths::EngineDir() / TEXT("Binaries") / TEXT("ThirdParty") / TEXT("PresentMon") / TEXT("Win64") / TEXT("PresentMon-1.8.0-x64.exe");
 
-	FString Arguments = 
-		FString::Printf(TEXT("-session_name session_%d -output_stdout -terminate_on_proc_exit -stop_existing_session -process_id %d"), 
-		Process->PID, Process->PID);
+	const FFlipModeMonitorSessionId SessionId = AllocFlipModeMonitorSessionId();
+	if (SessionId == InvalidFlipModeSessionId)
+	{
+		UE_LOGFMT(LogSwitchboard, Error, "Failed to allocate flip mode monitor session");
+		return nullptr;
+	}
+
+	MonitorProcess->OnProcessEnded.AddLambda(
+		[this, SessionId]()
+		{
+			if (EnsureFlipModeSessionAvailable(SessionId))
+			{
+				FreeFlipModeMonitorSessionId(SessionId);
+			}
+			else
+			{
+				UE_LOGFMT(LogSwitchboard, Error, "Unable to recover flip mode monitor session {SessionId}", SessionId);
+			}
+		}
+	);
+
+	const FString SessionName = MakeFlipModeMonitorSessionName(SessionId);
+
+	const FString Arguments = FString::Printf(
+		TEXT("-session_name %s -output_stdout -terminate_on_proc_exit -stop_existing_session -process_id %d"), 
+		*SessionName, Process->PID);
 
 	MonitorProcess->Handle = FSwitchboardProcess::CreateProc(
 		*MonitorProcess->Path,
@@ -1966,9 +2105,10 @@ FRunningProcess* FSwitchboardListener::FindOrStartFlipModeMonitorForUUID(const F
 		MonitorProcess->ClosePipes();
 
 		// Log error
-		const FString ErrorMsg = FString::Printf(TEXT("Could not start FlipMode monitor  %s"), *MonitorProcess->Path);
+		const FString ErrorMsg = FString::Printf(TEXT("Could not start FlipMode monitor %s"), *MonitorProcess->Path);
 		UE_LOG(LogSwitchboard, Error, TEXT("%s"), *ErrorMsg);
 
+		FreeFlipModeMonitorSessionId(SessionId);
 		return nullptr;
 	}
 
@@ -2579,6 +2719,8 @@ void FSwitchboardListener::HandleRunningProcesses(TArray<TSharedPtr<FRunningProc
 					}
 				}
 
+				Process->OnProcessEnded.Broadcast();
+
 				FPlatformProcess::CloseProc(Process->Handle);
 				Process->ClosePipes();
 
diff --git a/Engine/Source/Programs/SwitchboardListener/SblCore/Private/Windows/WindowsSwitchboardProcess.cpp b/Engine/Source/Programs/SwitchboardListener/SblCore/Private/Windows/WindowsSwitchboardProcess.cpp
--- a/Engine/Source/Programs/SwitchboardListener/SblCore/Private/Windows/WindowsSwitchboardProcess.cpp
+++ b/Engine/Source/Programs/SwitchboardListener/SblCore/Private/Windows/WindowsSwitchboardProcess.cpp
@@ -65,7 +65,8 @@ FWindowsSwitchboardProcess::FCreateProcResult FWindowsSwitchboardProcess::Create
 
 	// initialize startup info
 	CreateFlags |= EXTENDED_STARTUPINFO_PRESENT;
-	STARTUPINFOEXW StartupInfoEx = { 0 };
+	STARTUPINFOEXW StartupInfoEx;
+	FMemory::Memzero(StartupInfoEx);
 	StartupInfoEx.StartupInfo = {
 		sizeof(STARTUPINFOEXW),
 		nullptr,
