SwitchboardListener: Fix soft termination of process tree.

It wasn't effectively performing a soft termination (in Windows) because taskkill was using /F. Even without /F and only /T, it would also not working because /T fails when any of the child processes require /F (such as node.exe, as it turns out).

Instead we now issue a taskkill on each process in the tree, then after the timeout we force terminate any process that is still running.

#rnx

diff --git a/Engine/Source/Programs/SwitchboardListener/SblCore/Private/SwitchboardListener.cpp b/Engine/Source/Programs/SwitchboardListener/SblCore/Private/SwitchboardListener.cpp
--- a/Engine/Source/Programs/SwitchboardListener/SblCore/Private/SwitchboardListener.cpp
+++ b/Engine/Source/Programs/SwitchboardListener/SblCore/Private/SwitchboardListener.cpp
@@ -39,10 +39,11 @@
 #pragma warning(disable : 4005)	// Disable macro redefinition warning for compatibility with Windows SDK 8+
 
 #include "Windows/AllowWindowsPlatformTypes.h"
+#include <evntrace.h>
+#include <shellapi.h>
+#include <tlhelp32.h>
 #include <Windows.h>
 #include <WinUser.h>
-#include <shellapi.h>
-#include <evntrace.h>
 #include "Windows/HideWindowsPlatformTypes.h"
 
 #include "ScopedNvApi.h"
@@ -1267,45 +1268,162 @@ bool FSwitchboardListener::Task_FixExeFlags(const FSwitchboardFixExeFlagsTask& T
 }
 
 
-bool FSwitchboardListener::KillProcessNow(FRunningProcess* InProcess, float SoftKillTimeout)
+#if PLATFORM_WINDOWS
+
+/**
+ * Recursively collects the PIDs of all child processes of a given root PID.
+ */
+static void WindowsGatherChildPIDs(uint32 RootPID, TArray<uint32>& OutPIDs)
 {
-	SWITCHBOARD_TRACE_CPUPROFILER_EVENT_SCOPE(FSwitchboardListener::KillProcessNow);
+	const HANDLE SnapshotHandle = CreateToolhelp32Snapshot(TH32CS_SNAPPROCESS, 0);
 
-	if (InProcess && InProcess->Handle.IsValid() && FPlatformProcess::IsProcRunning(InProcess->Handle))
+	if (SnapshotHandle == INVALID_HANDLE_VALUE)
 	{
-		UE_LOG(LogSwitchboard, Display, TEXT("Killing app with PID %d"), InProcess->PID);
+		return;
+	}
 
-#if PLATFORM_WINDOWS
-		// try a soft kill first
-		if(SoftKillTimeout > 0)
+	PROCESSENTRY32 ProcessEntry;
+	ProcessEntry.dwSize = sizeof(ProcessEntry);
+
+	if (Process32First(SnapshotHandle, &ProcessEntry))
+	{
+		do
 		{
-			const FString Params = FString::Printf(TEXT("/T /F /PID %d"), InProcess->PID);
+			if (ProcessEntry.th32ParentProcessID == RootPID)
+			{
+				OutPIDs.Add(ProcessEntry.th32ProcessID);
+				WindowsGatherChildPIDs(ProcessEntry.th32ProcessID, OutPIDs); // Recursion
+			}
+		} while (Process32Next(SnapshotHandle, &ProcessEntry));
+	}
+
+	CloseHandle(SnapshotHandle);
+}
+
+
+/** 
+ * Soft-terminates a process by PID, excluding its process tree. 
+ */
+static bool WindowsSoftTerminateByPID(uint32 PID)
+{
+	const FString Params = FString::Printf(TEXT("/PID %u"), PID);
 
-			FString OutStdOut;
+	FString StdOut;
+	int32 ReturnCode = 0;
 
-			FPlatformProcess::ExecProcess(TEXT("TASKKILL"), *Params, nullptr, &OutStdOut, nullptr);
+	FPlatformProcess::ExecProcess(TEXT("TASKKILL"), *Params, &ReturnCode, &StdOut, nullptr);
 
-			const double TimeoutTime = FPlatformTime::Seconds() + SoftKillTimeout;
-			const float SleepTime = 0.050f;
+	return (ReturnCode == 0);
+}
+
+
+/**
+ * Terminate a process tree by first sending individual soft sins to each PID,
+ * waiting for up to SoftKillTimeout seconds, then forcefully terminating survivors.
+ */
+bool WindowsKillProcessNow(FRunningProcess* InProcess, float SoftKillTimeout)
+{
+	check(InProcess && InProcess->Handle.IsValid()); // We expect this was checked before calling this platform specific function.
+
+	// Gather all PIDs to terminate, including the root.
+
+	TArray<uint32> AllPIDs;
+	AllPIDs.Add(InProcess->PID);
+
+	WindowsGatherChildPIDs(InProcess->PID, AllPIDs);
+
+	// Soft terminate each one
+
+	for (uint32 PID : AllPIDs)
+	{
+		WindowsSoftTerminateByPID(PID);
+	}
 
-			while(FPlatformTime::Seconds() < TimeoutTime && FPlatformProcess::IsProcRunning(InProcess->Handle))
+	// Wait for them to terminate, up to the given timeout
+
+	const double DeadlineTime = FPlatformTime::Seconds() + SoftKillTimeout;
+
+	while (FPlatformTime::Seconds() < DeadlineTime)
+	{
+		bool bAnyStillRunning = false;
+
+		for (uint32 PID : AllPIDs)
+		{
+			const HANDLE ProcessHandle = ::OpenProcess(SYNCHRONIZE, false, PID);
+
+			if (ProcessHandle)
 			{
-				FPlatformProcess::Sleep(SleepTime);
+				DWORD ExitCode = 0;
+
+				if (::GetExitCodeProcess(ProcessHandle, &ExitCode) && ExitCode == STILL_ACTIVE)
+				{
+					bAnyStillRunning = true;
+				}
+
+				::CloseHandle(ProcessHandle);
+			}
+
+			if (bAnyStillRunning)
+			{
+				break;
 			}
 		}
-#endif // PLATFORM_WINDOWS
 
-		if (FPlatformProcess::IsProcRunning(InProcess->Handle))
+		if (!bAnyStillRunning)
 		{
-			const bool bKillTree = true;
-			FPlatformProcess::TerminateProc(InProcess->Handle, bKillTree);
+			break;
 		}
 
-		// Pipes will be closed in HandleRunningProcesses
-		return true;
+		FPlatformProcess::Sleep(0.05f);
 	}
 
-	return false;
+	// Hard‐terminate surviving PIDs
+
+	for (uint32 PID : AllPIDs)
+	{
+		FProcHandle TermHandle = FPlatformProcess::OpenProcess(PID);
+
+		if (TermHandle.IsValid())
+		{
+			if (FPlatformProcess::IsProcRunning(TermHandle))
+			{
+				FPlatformProcess::TerminateProc(TermHandle, false /* bKillTree */);
+			}
+			FPlatformProcess::CloseProc(TermHandle);
+		}
+	}
+
+	return true;
+}
+#endif // PLATFORM_WINDOWS
+
+
+bool FSwitchboardListener::KillProcessNow(FRunningProcess* InProcess, float SoftKillTimeout)
+{
+	SWITCHBOARD_TRACE_CPUPROFILER_EVENT_SCOPE(FSwitchboardListener::KillProcessNow);
+
+	if (!InProcess || !InProcess->Handle.IsValid() || !FPlatformProcess::IsProcRunning(InProcess->Handle))
+	{
+		return false;
+	}
+
+	UE_LOG(LogSwitchboard, Display, TEXT("Killing app with PID %d"), InProcess->PID);
+
+#if PLATFORM_WINDOWS
+
+	return WindowsKillProcessNow(InProcess, SoftKillTimeout);
+
+#else
+
+	if (FPlatformProcess::IsProcRunning(InProcess->Handle))
+	{
+		FPlatformProcess::TerminateProc(InProcess->Handle, true /* bKillTree */);
+	}
+
+	// Pipes will be closed in HandleRunningProcesses
+	return true;
+
+#endif // PLATFORM_WINDOWS
 }
 
 bool FSwitchboardListener::Task_ReceiveFileFromClient(const FSwitchboardReceiveFileFromClientTask& InReceiveFileFromClientTask)
