Switchboard: Fix bad structured logging format strings in PosixOSSwitchboardCredential.

#rnx

diff --git a/Engine/Source/Programs/SwitchboardListener/SblCore/Private/PosixOS/PosixOSSwitchboardCredential.cpp b/Engine/Source/Programs/SwitchboardListener/SblCore/Private/PosixOS/PosixOSSwitchboardCredential.cpp
--- a/Engine/Source/Programs/SwitchboardListener/SblCore/Private/PosixOS/PosixOSSwitchboardCredential.cpp
+++ b/Engine/Source/Programs/SwitchboardListener/SblCore/Private/PosixOS/PosixOSSwitchboardCredential.cpp
@@ -171,7 +171,7 @@ namespace UE::SwitchboardListener::Private
 		if (!ensure(TempFd != -1))
 		{
 			const int ErrorCode = errno;
-			UE_LOGFMT(LogSwitchboard, Error, "{Func}: mkstemp failed (errno: %d; path: '%S')",
+			UE_LOGFMT(LogSwitchboard, Error, "{Func}: mkstemp failed (errno: {errno}; path: '{TempFilePath}')",
 				__FUNCTION__, ErrorCode, TempFilePathU8);
 			return false;
 		}
@@ -200,7 +200,7 @@ namespace UE::SwitchboardListener::Private
 		if (!ensure(CloseResult == 0))
 		{
 			const int CloseErrorCode = errno;
-			UE_LOGFMT(LogSwitchboard, Warning, "{Func}: close failed (errno: %d; path: '%S')",
+			UE_LOGFMT(LogSwitchboard, Warning, "{Func}: close failed (errno: {errno}; path: '{TempFilePath}')",
 				__FUNCTION__, CloseErrorCode, TempFilePathU8);
 		}
 
@@ -209,7 +209,7 @@ namespace UE::SwitchboardListener::Private
 		{
 			PlatformFile.DeleteFile(UTF8_TO_TCHAR(TempFilePathU8.GetData()));
 			UE_LOGFMT(LogSwitchboard, Error,
-				"{Func}: write failed (errno: %d; written %d/%d; path: '%S')",
+				"{Func}: write failed (errno: {errno}; written {BytesWritten}/{JsonByteLen}; path: '{TempFilePath}')",
 				__FUNCTION__, ErrorCode, TotalBytesWritten, JsonString.Len(), TempFilePathU8);
 			return false;
 		}
@@ -221,7 +221,7 @@ namespace UE::SwitchboardListener::Private
 		{
 			ErrorCode = errno;
 			PlatformFile.DeleteFile(UTF8_TO_TCHAR(TempFilePathU8.GetData()));
-			UE_LOGFMT(LogSwitchboard, Warning, "{Func}: rename failed (errno: %d; path: '%S' -> '%s')",
+			UE_LOGFMT(LogSwitchboard, Warning, "{Func}: rename failed (errno: {errno}; path: '{TempFilePath}' -> '{DestFilePath}')",
 				__FUNCTION__, ErrorCode, TempFilePathU8, DestFilePath);
 			return false;
 		}
